Locomotor Plugin: Adding new experimental procedural locomotion plugin

diff --git a/Engine/Plugins/Experimental/Animation/Locomotor/Locomotor.uplugin b/Engine/Plugins/Experimental/Animation/Locomotor/Locomotor.uplugin
new file mode 100644
--- /dev/null
+++ b/Engine/Plugins/Experimental/Animation/Locomotor/Locomotor.uplugin
@@ -0,0 +1,30 @@
+{
+	"FileVersion": 3,
+	"Version": 1,
+	"VersionName": "1.0",
+	"FriendlyName": "Locomotor",
+	"Description": "Procedural animation for Control Rig.",
+	"Category": "Animation",
+	"CreatedBy": "Epic Games",
+	"CreatedByURL": "www.epicgames.com",
+	"DocsURL": "",
+	"MarketplaceURL": "",
+	"SupportURL": "",
+	"CanContainContent": false,
+	"IsBetaVersion": true,
+	"IsExperimentalVersion": true,
+	"Installed": false,
+	"Plugins": [
+		{
+			"Name": "ControlRig",
+			"Enabled": true
+		}
+	],
+	"Modules": [
+		{
+			"Name": "Locomotor",
+			"Type": "Runtime",
+			"LoadingPhase": "PreDefault"
+		}
+	]
+}
\ No newline at end of file
diff --git a/Engine/Plugins/Experimental/Animation/Locomotor/Source/Locomotor/Locomotor.Build.cs b/Engine/Plugins/Experimental/Animation/Locomotor/Source/Locomotor/Locomotor.Build.cs
new file mode 100644
--- /dev/null
+++ b/Engine/Plugins/Experimental/Animation/Locomotor/Source/Locomotor/Locomotor.Build.cs
@@ -0,0 +1,55 @@
+// Copyright Epic Games, Inc. All Rights Reserved.
+
+using UnrealBuildTool;
+
+public class Locomotor : ModuleRules
+{
+	public Locomotor(ReadOnlyTargetRules Target) : base(Target)
+	{
+		PCHUsage = ModuleRules.PCHUsageMode.UseExplicitOrSharedPCHs;
+		
+		PublicIncludePaths.AddRange(
+			new string[] {
+				// ... add public include paths required here ...
+			}
+			);
+				
+		
+		PrivateIncludePaths.AddRange(
+			new string[] {
+				// ... add other private include paths required here ...
+			}
+			);
+			
+		
+		PublicDependencyModuleNames.AddRange(
+			new string[]
+			{
+				"Core",
+				"RigVM",
+				"ControlRig",
+				// ... add other public dependencies that you statically link with here ...
+			}
+			);
+			
+		
+		PrivateDependencyModuleNames.AddRange(
+			new string[]
+			{
+				"CoreUObject",
+				"Engine",
+				"Slate",
+				"SlateCore",
+				// ... add private dependencies that you statically link with here ...	
+			}
+			);
+		
+		
+		DynamicallyLoadedModuleNames.AddRange(
+			new string[]
+			{
+				// ... add any modules that your module loads dynamically here ...
+			}
+			);
+	}
+}
diff --git a/Engine/Plugins/Experimental/Animation/Locomotor/Source/Locomotor/Private/Locomotor.cpp b/Engine/Plugins/Experimental/Animation/Locomotor/Source/Locomotor/Private/Locomotor.cpp
new file mode 100644
--- /dev/null
+++ b/Engine/Plugins/Experimental/Animation/Locomotor/Source/Locomotor/Private/Locomotor.cpp
@@ -0,0 +1,20 @@
+// Copyright Epic Games, Inc. All Rights Reserved.
+
+#include "Locomotor.h"
+
+#define LOCTEXT_NAMESPACE "FLocomotorModule"
+
+void FLocomotorModule::StartupModule()
+{
+	// This code will execute after your module is loaded into memory; the exact timing is specified in the .uplugin file per-module
+}
+
+void FLocomotorModule::ShutdownModule()
+{
+	// This function may be called during shutdown to clean up your module.  For modules that support dynamic reloading,
+	// we call this function before unloading the module.
+}
+
+#undef LOCTEXT_NAMESPACE
+	
+IMPLEMENT_MODULE(FLocomotorModule, Locomotor)
\ No newline at end of file
diff --git a/Engine/Plugins/Experimental/Animation/Locomotor/Source/Locomotor/Private/LocomotorCore.cpp b/Engine/Plugins/Experimental/Animation/Locomotor/Source/Locomotor/Private/LocomotorCore.cpp
new file mode 100644
--- /dev/null
+++ b/Engine/Plugins/Experimental/Animation/Locomotor/Source/Locomotor/Private/LocomotorCore.cpp
@@ -0,0 +1,992 @@
+﻿// Copyright Epic Games, Inc. All Rights Reserved.
+
+#include "LocomotorCore.h"
+
+#include "Engine/World.h"
+#include "Components/PrimitiveComponent.h"
+#include "CollisionQueryParams.h"
+#include "CollisionShape.h"
+
+//#pragma optimize("", off)
+
+FLocomotorFoot::FLocomotorFoot(
+	const FTransform& InitialWorldRootGoal,
+	const FTransform& InitialWorldFoot,
+	const FLocomotorFootSettings& InSettings)
+{
+	Settings = InSettings;
+	CurrentPhase = 1.0f;
+	FootPhaseWhenSwingEnds = 0.5f;
+	InitialRelativeToRootGoal = InitialWorldFoot.GetRelativeTransform(InitialWorldRootGoal);
+	InitialRelativeToBody = FTransform::Identity; // this is initialized after Pelvis calculates initial body transform
+	InitialWorld = InitialWorldFoot;
+	CurrentWorld = InitialWorld;
+	CurrentWorldFlatPosition = InitialWorld.GetLocation();
+	CurrentWorldFlatPositionNoEase = InitialWorld.GetLocation();
+	CurrentTargetWorld = InitialWorld;
+	PrevWorld = InitialWorld;
+	CurrentRotationSpring.Reset(InitialWorld.GetRotation());
+	PlantedWorld = InitialWorld;
+	FinalTargetWorld = InitialWorld;
+	StrideOriginWorld = InitialWorld;
+	CurrentHeightOffset = 0.f;
+	bAtRest = true;
+	bInSwingPhase = false;
+	bWantsToStep = false;
+}
+
+int32 FLocomotorFootSet::AddFoot(
+	const FTransform& InitialWorldRootGoal,
+	const FTransform& InitialWorldFoot,
+	const FLocomotorFootSettings& InSettings)
+{
+	const int32 FootIndexInSet = Feet.Emplace(InitialWorldRootGoal, InitialWorldFoot, InSettings);
+	FLocomotorFoot& NewFoot = Feet[FootIndexInSet];
+	NewFoot.Settings.StaticPhaseOffset = FootIndexInSet % 2 == 0.f ? 0.f : 0.5f;
+	NewFoot.Settings.FootSetIndex = SetIndex;
+	return FootIndexInSet;
+}
+
+void FLocomotorHead::Initialize(const FTransform& InInitialWorld)
+{
+	InitialWorld = InInitialWorld;
+	bInitialized = true;
+}
+
+void FLocomotorPelvis::Initialize(const TArray<FLocomotorFoot*>& AllFeet, const FTransform& InitialRootGoalWorld)
+{
+	CurrentWorld = InitialWorld; // initial pelvis world set when Reset()'ing the locomotor
+	CalculateBodyTransform(AllFeet, InitialRootGoalWorld, InitialBodyWorld);
+	CurrentBodyWorld = CurrentBodyTargetWorld = PreviousBodyTargetWorld = InitialBodyWorld;
+
+	BodyLeadDamper.Reset(FVector::ZeroVector); // initialized to zero because we smooth the lead amount, not absolute value
+	BodyRotationSpring.Reset(InitialBodyWorld.GetRotation());
+
+	CurrentHeightOffset = InitialWorld.GetLocation().Z;
+	BobOffsetSpring.Reset(0.f);
+}
+
+void FLocomotorPelvis::CalculateBodyTransform(
+	const TArray<FLocomotorFoot*> AllFeet,
+	const FTransform& RootGoal,
+	FTransform& OutBodyTransform)
+{
+	FVector AvgFootPosition = FVector::ZeroVector;
+	for (const FLocomotorFoot* FootPtr : AllFeet)
+	{
+		AvgFootPosition += FootPtr->CurrentWorldFlatPositionNoEase;
+	}
+	const float InvNumFeet = 1.0f / AllFeet.Num();
+	AvgFootPosition *= InvNumFeet;
+
+	OutBodyTransform.SetRotation(RootGoal.GetRotation());
+	OutBodyTransform.SetTranslation(AvgFootPosition);
+}
+
+void FLocomotor::Reset(const FTransform& InInitialRootGoalWorld, const FTransform& InInitialWorldPelvis)
+{
+	InitialRootGoalWorld = InInitialRootGoalWorld;
+	Pelvis.InitialWorld = InInitialWorldPelvis;
+	Pelvis.PelvisPositionDamper.Reset(InInitialWorldPelvis.GetLocation());
+	
+	AllFeet.Reset();
+	FootSets.Reset();
+	CurrentGlobalPhase = 0.0f;
+	CurrentSpeed = 0.0f;
+	CurrentPhaseSpeed = 0.0f;
+	CurrentPercentOfMaxSpeed = 0.0f;
+	bFullyAtRest = true;
+	bAccelerating = true;
+	AccumulatedTimeToSimulate = 0.0f;
+
+	bPostInitialized = false;
+}
+
+int32 FLocomotor::AddFootSet(const float PhaseOffset)
+{
+	bPostInitialized = false;
+	return FootSets.Emplace(FLocomotorFootSet(PhaseOffset, FootSets.Num()));
+}
+
+int32 FLocomotor::AddFootToSet(
+	const int32 FootSetIndex,
+	const FTransform& InitialWorldFoot,
+	const FLocomotorFootSettings& InSettings)
+{
+	bPostInitialized = false;
+	
+	if (!FootSets.IsValidIndex(FootSetIndex))
+	{
+		return INDEX_NONE;
+	}
+	
+	const int32 FootIndexInSet = FootSets[FootSetIndex].AddFoot(InitialRootGoalWorld, InitialWorldFoot, InSettings);
+	AllFeet.Add(&FootSets[FootSetIndex].GetFeet()[FootIndexInSet]);
+	return FootIndexInSet;
+}
+
+void FLocomotor::SetHead(const FTransform& InitialWorldTransform)
+{
+	Head.Initialize(InitialWorldTransform);
+}
+
+void FLocomotor::RunSimulation(const FLocomotorInputSettings& InSettings)
+{
+	// run post initialization and validation
+	if (!Initialize())
+	{
+		return;
+	}
+
+	// copy latest input settings
+	Settings = InSettings;
+
+	// run simulation at max internal time step...
+	// this breaks large time steps into multiple smaller simulation steps for robustness
+	constexpr float MaxTimeStep = 1 / 120.0f;
+	constexpr float MinTimeStep = MaxTimeStep / 4.f;
+	// delta time is accumulated across updates, allows for extreme slow motion
+	AccumulatedTimeToSimulate += InSettings.DeltaTime;
+	while (AccumulatedTimeToSimulate > MinTimeStep)
+	{
+		Settings.DeltaTime = AccumulatedTimeToSimulate >= MaxTimeStep ? MaxTimeStep : AccumulatedTimeToSimulate;
+		AccumulatedTimeToSimulate -= Settings.DeltaTime;
+		Simulate();
+	}
+}
+
+void FLocomotor::GetFeetCurrent(TArray<FTransform> OutTransforms) const
+{
+	// OutTransforms must be sized by calling code
+	if (!ensure(OutTransforms.Num() == AllFeet.Num()))
+	{
+		return;
+	}
+	
+	for (int32 FootIndex = 0; FootIndex < AllFeet.Num(); ++FootIndex)
+	{
+		OutTransforms[FootIndex] = AllFeet[FootIndex]->CurrentWorld;
+	}
+}
+
+void FLocomotor::Simulate()
+{
+	UpdateWorldSpeedAndPhase();
+	UpdateFeetTargets();
+	AnimateFeet();
+	UpdateBody();
+	AnimatePelvis();
+}
+
+bool FLocomotor::Initialize()
+{
+	if (!bPostInitialized)
+	{
+		Pelvis.Initialize(AllFeet, InitialRootGoalWorld);
+
+		// requires Pelvis.Initialize() to be called first to get InitialBodyWorld
+		for (FLocomotorFoot* Foot : AllFeet)
+		{
+			Foot->InitialRelativeToBody = Foot->InitialWorld.GetRelativeTransform(Pelvis.InitialBodyWorld);
+		}
+	
+		bPostInitialized = true;
+	}
+	
+	return !GetFeet().IsEmpty();
+}
+
+void FLocomotor::UpdateWorldSpeedAndPhase()
+{
+	//
+	// apply acceleration and deceleration
+	//
+
+	// how far away is the pelvis from it's target location?
+	const FTransform InitialPelvisRelativeToRootGoal = Pelvis.InitialWorld.GetRelativeTransform(InitialRootGoalWorld);
+	const FTransform CurrentPelvisAttachedToRootGoal = InitialPelvisRelativeToRootGoal * Settings.CurrentWorldRootGoal;
+	const float DistanceToGoal = FVector::Distance(CurrentPelvisAttachedToRootGoal.GetLocation(), Pelvis.CurrentWorld.GetLocation());
+	const bool GoalFurtherThanSingleStride = DistanceToGoal > Settings.Movement.MinimumStepLength * 2.0f;
+	
+	// how long would it take to slow down to minimum speed?
+	// calculate distance to reach minimum velocity from current velocity (given current deceleration)
+	const float CurrentStoppingDistance =
+		(FMath::Pow(Settings.Movement.SpeedMin, 2) - FMath::Pow(CurrentSpeed, 2))	/ (-Settings.Movement.Deceleration * 2.f);
+	const bool bGoalBeyondStoppingDistance = DistanceToGoal > CurrentStoppingDistance;
+
+	// determine if body is rotating a significant amount
+	//const float AngularDistanceRadians = InitialRootGoalWorld.GetRotation().AngularDistance(Settings.CurrentWorldRootGoal.GetRotation());
+	//constexpr float AngularThresholdForFullAcceleration = 40.0f;
+	//const bool bNeedsFastRotation = FMath::RadiansToDegrees(AngularDistanceRadians) > AngularThresholdForFullAcceleration;
+
+	// determine whether we are speeding up or slowing down based current speed and distance to the target
+	bAccelerating = GoalFurtherThanSingleStride && bGoalBeyondStoppingDistance;// || bNeedsFastRotation;
+	const float Acceleration = bAccelerating ? Settings.Movement.Acceleration : -Settings.Movement.Deceleration;
+	// integrate the acceleration for this time step
+	CurrentSpeed = CurrentSpeed + Acceleration * Settings.DeltaTime;
+	// clamp within bounds
+	CurrentSpeed = FMath::Clamp(CurrentSpeed, Settings.Movement.SpeedMin, Settings.Movement.SpeedMax);
+	// where are we in the range betweeen min and max speed
+	CurrentPercentOfMaxSpeed = FMath::GetRangePct(Settings.Movement.SpeedMin, Settings.Movement.SpeedMax, CurrentSpeed);
+	CurrentPercentOfMaxSpeed = FMath::Clamp(CurrentPercentOfMaxSpeed,0.f, 1.f);
+	CurrentPhaseSpeed = FMath::Lerp(Settings.Movement.PhaseSpeedMin, Settings.Movement.PhaseSpeedMax, CurrentPercentOfMaxSpeed);
+	CurrentStrideLength = CurrentSpeed / CurrentPhaseSpeed;
+	
+	// reset phase if fully at rest
+	if (bFullyAtRest)
+	{
+		// stop updating at rest
+		CurrentGlobalPhase = 0.0f;
+	}
+	else
+	{
+		// update the global phase
+		CurrentGlobalPhase += CurrentPhaseSpeed * Settings.DeltaTime;
+		CurrentGlobalPhase = WrapPhaseInRange(CurrentGlobalPhase);
+	}
+
+	// DEBUG OUTPUT
+	//UE_LOG(LogTemp, Warning, TEXT("Distance to Goal  : %f"), DistanceToGoal);
+	//UE_LOG(LogTemp, Warning, TEXT("Stopping Distance : %f"), CurrentStoppingDistance);
+	//UE_LOG(LogTemp, Warning, TEXT("Current speed     : %f"), CurrentSpeed);
+	//UE_LOG(LogTemp, Warning, TEXT("Phase   speed     : %f"), CurrentPhaseSpeed);
+	//UE_LOG(LogTemp, Warning, TEXT("Delta time        : %f"), Settings.DeltaTime);
+	//UE_LOG(LogTemp, Warning, TEXT("PercentOfMaxSpeed : %f"), CurrentPercentOfMaxSpeed);
+	//UE_LOG(LogTemp, Warning, TEXT("Current Stride    : %f"), CurrentStrideLength);
+	//UE_LOG(LogTemp, Warning, TEXT("Accelerating      : %s"), bAccelerating ? TEXT("Accelerating") : TEXT("Decelerating"));
+}
+
+void FLocomotor::UpdateFeetTargets()
+{
+	// continuously update final targets for each foot
+	// these are the foot locations at the root goal location
+	for (FLocomotorFoot* FootPtr : AllFeet)
+	{
+		FLocomotorFoot& Foot = *FootPtr;
+		
+		// continuously update the final foot target
+		Foot.FinalTargetWorld = Foot.InitialRelativeToRootGoal * Settings.CurrentWorldRootGoal;
+		if (Settings.Stepping.bEnableGroundCollision)
+		{
+			ProjectToGroundWithSphereCast(Foot, Foot.FinalTargetWorld);	
+		}
+
+		// vector from prev planted location to final target
+		Foot.PrevToFinalTargetNorm = Foot.FinalTargetWorld.GetLocation() - Foot.PlantedWorld.GetLocation();
+		Foot.PrevToFinalTargetDistance = Foot.PrevToFinalTargetNorm.Length();
+		Foot.PrevToFinalTargetNorm.Normalize();
+	}
+
+	// continuously update "bWantsToStep"
+	// (based on if the foot's final target is beyond the min stride distance)
+	bool bAnyFeetWantToStep = false;
+	for (FLocomotorFoot* FootPtr : AllFeet)
+	{
+		FLocomotorFoot& Foot = *FootPtr;
+		// wants to step if the distance to the target is greater than the minimum stride length
+		Foot.bWantsToStep = Foot.PrevToFinalTargetDistance > Settings.Movement.MinimumStepLength;
+		if (Foot.bWantsToStep)
+		{
+			bAnyFeetWantToStep = true;
+		}
+	}
+	
+	// if we are fully at rest and any foot wants to step,
+	// then we need to reset the global phase to the closest foot AND reset each foot's current phase
+	if (bFullyAtRest && bAnyFeetWantToStep)
+	{
+		// find the foot that is furthest from the target
+		float FurthestDistance = -1.0f;
+		int32 FurthestFootIndex = INDEX_NONE;
+		for (int32 FootIndex=0; FootIndex<AllFeet.Num(); ++FootIndex)
+		{
+			const FLocomotorFoot* Foot = AllFeet[FootIndex];
+			if (Foot->PrevToFinalTargetDistance > FurthestDistance)
+			{
+				FurthestDistance = Foot->PrevToFinalTargetDistance;
+				FurthestFootIndex = FootIndex;
+			}
+		}
+
+		// rewind/fast-forward the current phase so that the furthest foot steps first
+		CurrentGlobalPhase = WrapPhaseInRange(-AllFeet[FurthestFootIndex]->Settings.StaticPhaseOffset);
+
+		// force rewind/fast-forward all foot phases
+		for (FLocomotorFoot* FootPtr : AllFeet)
+		{
+			FLocomotorFoot& Foot = *FootPtr;
+			Foot.CurrentPhase = Foot.TargetPhase = WrapPhaseInRange(CurrentGlobalPhase + Foot.Settings.StaticPhaseOffset);
+			Foot.FootPhaseWhenSwingEnds = 0.5f;
+			//UE_LOG(LogTemp, Warning, TEXT("Force updating feet phases: %f"), Foot.CurrentPhase);
+		}
+	}
+
+	// update "bAtRest", "bInSwingPhase" and "bUnplantingThisTick" states
+	for (FLocomotorFoot* FootPtr : AllFeet)
+	{
+		FLocomotorFoot& Foot = *FootPtr;
+		
+		// record if this foot was planted before this tick
+		const bool bWasInSwing = Foot.bInSwingPhase;
+		// check if the foot is now in the swing phase
+		Foot.bInSwingPhase = Foot.CurrentPhase < Foot.FootPhaseWhenSwingEnds;
+		// if the foot is in it's planted phase, AND the target is not beyond the min threshold, then we can consider this foot "at rest"
+		Foot.bAtRest = !Foot.bInSwingPhase && !Foot.bWantsToStep;
+		// is this foot unplanting this tick?
+		Foot.bUnplantingThisTick = !bWasInSwing && Foot.bInSwingPhase; 
+	}
+
+	// update "CurrentPhase" and "PhaseWhenSwingEnds"
+	const float PhaseCatchUpPerSecond = CurrentPhaseSpeed * 2.0f; 
+	const float PhaseBlendRate = PhaseCatchUpPerSecond * Settings.DeltaTime;
+	for (FLocomotorFoot* FootPtr : AllFeet)
+	{
+		FLocomotorFoot& Foot = *FootPtr;
+		
+		// update feet phases
+		if (Foot.bAtRest)
+		{
+			Foot.CurrentPhase = Foot.TargetPhase = 1.0f;
+		}
+		else
+		{
+			const float PhaseOffsetOfSet = FootSets[Foot.Settings.FootSetIndex].GetPhaseOffset();
+			const float StaticOffsetOfFoot = Foot.Settings.StaticPhaseOffset;
+			Foot.TargetPhase = WrapPhaseInRange(CurrentGlobalPhase + StaticOffsetOfFoot + PhaseOffsetOfSet);
+			Foot.CurrentPhase = BlendTowardsTargetPhase(Foot.CurrentPhase, Foot.TargetPhase, PhaseBlendRate);
+			//UE_LOG(LogTemp, Warning, TEXT("Updating foot phase Current/Target: %f, %f"), Foot.CurrentPhase, Foot.TargetPhase);
+		}
+		
+		// if unplanting this tick, we need to calculate the phase when swing ends (locked it in for the duration of the step)
+		if (Foot.bUnplantingThisTick)
+		{
+			Settings.Stepping.PercentOfStrideInAir = FMath::Clamp(Settings.Stepping.PercentOfStrideInAir, 0.1f, 0.95f);
+			Settings.Stepping.AirExtensionAtMaxSpeed = FMath::Clamp(Settings.Stepping.AirExtensionAtMaxSpeed,0.0f, 0.95f);
+			
+			const float PercentStrideInAirAtMinSpeed = Settings.Stepping.PercentOfStrideInAir;
+			const float PercentStrideInAirAtMaxSpeed = FMath::Clamp(
+				PercentStrideInAirAtMinSpeed + Settings.Stepping.AirExtensionAtMaxSpeed,
+				Settings.Stepping.PercentOfStrideInAir,
+				0.95f);
+
+			Foot.FootPhaseWhenSwingEnds = FMath::Lerp(
+				PercentStrideInAirAtMinSpeed,
+				PercentStrideInAirAtMaxSpeed,
+				CurrentPercentOfMaxSpeed);
+		}
+	}
+	
+	// continuously update the current target of each foot that is in a swing phase
+	for (FLocomotorFoot* FootPtr : AllFeet)
+	{
+		FLocomotorFoot& Foot = *FootPtr;
+		
+		if (!Foot.bInSwingPhase)
+		{
+			continue;
+		}
+
+		// some foot parameters are updated once at the start of a stride (when unplanting)
+		if (Foot.bUnplantingThisTick)
+		{
+			// update stride origin when lifting off
+			// NOTE: we don't use the Foot.PlantedWorldPosition because it may be behind where the foot's neutral pose can reach
+			Foot.StrideOriginWorld = Foot.InitialRelativeToBody * Pelvis.CurrentBodyWorld;
+			ProjectToGroundWithSphereCast(Foot, Foot.StrideOriginWorld);
+
+			// update foot height for the current stride
+			// NOTE: this is not done continuously because it can lead to foot height wobbling if the speed changes mid step
+			const float MinStrideHeight = Settings.Stepping.StepHeight * 0.2f;
+			const float MaxStrideHeight = Settings.Stepping.StepHeight;
+			Foot.CurrentStrideHeight = FMath::Lerp(MinStrideHeight, MaxStrideHeight, CurrentPercentOfMaxSpeed);
+		}
+
+		// generate a candidate target transform for the foot
+		FTransform NewTarget = FTransform::Identity;
+		
+		// target location is in direct line from stride origin to final target, scaled by step length
+		const FVector StrideOriginToFinalTarget = Foot.FinalTargetWorld.GetLocation() - Foot.StrideOriginWorld.GetLocation();
+		const float StepLengthToUse = FMath::Min(StrideOriginToFinalTarget.Length(), CurrentStrideLength);
+		NewTarget.SetTranslation(Foot.StrideOriginWorld.GetLocation() + StrideOriginToFinalTarget.GetSafeNormal() * StepLengthToUse);
+
+		// (optionally) project the candidate target onto the ground and update orientation based on normal
+		if (Settings.Stepping.bEnableGroundCollision)
+		{
+			FVector GroundNormal = ProjectToGroundWithSphereCast(Foot, NewTarget);
+
+			// update pitch/roll orientation of the foot based on current target ground normal
+
+			// decompose normal into pitch and roll rotations
+			const FVector Fwd = Settings.CurrentWorldRootGoal.GetRotation().RotateVector(FVector(0,1,0));
+			const FVector Up = FVector(0,0,1);
+			FVector Side;
+			FQuat PitchRotation;
+			FQuat RollRotation;
+			CalcPitchRollFromNormal(Fwd, Up, GroundNormal, Side,PitchRotation, RollRotation);
+
+			// scale rotations by user amount
+			PitchRotation = FQuat::FastLerp(FQuat::Identity, PitchRotation, Settings.Stepping.OrientFootToGroundPitch);
+			RollRotation = FQuat::FastLerp(FQuat::Identity, RollRotation, Settings.Stepping.OrientFootToGroundRoll);
+
+			// apply to the current target
+			NewTarget.SetRotation(PitchRotation * RollRotation * Foot.FinalTargetWorld.GetRotation());
+		}
+		else
+		{
+			// target rotation is simply the FINAL target rotation
+			// while the location of the feet takes many steps to reach the final target,
+			// the rotation of the feet should get there in a single step.
+			NewTarget.SetRotation(Foot.FinalTargetWorld.GetRotation());
+		}
+		
+		// now continuously update the target as the phase progresses
+		// NOTE:
+		//	- we allow the FIRST target after un-planting to be FULLY applied, but afterwards we limit target updates by foot speed
+		//	- this roughly simulates "committing" to the step and the inability to change direction on a dime.
+		if (Foot.bUnplantingThisTick)
+		{
+			// retarget target spring at start of stride
+			Foot.CurrentTargetSpring.Reset(NewTarget.GetLocation());
+		}
+		else
+		{
+			// when the character is slowing down, new targets are pulled closer which may make them further from the final target
+			// we should never update a target to make it further from the goal
+			const float OldTargetToFinalDistSq = FVector::DistSquared(Foot.FinalTargetWorld.GetLocation(), Foot.CurrentTargetWorld.GetLocation());
+			const float NewTargetToFinalDistSq = FVector::DistSquared(Foot.FinalTargetWorld.GetLocation(), NewTarget.GetLocation());
+			if (NewTargetToFinalDistSq < OldTargetToFinalDistSq)
+			{
+				// move the current target towards the new target with a damped spring
+				// NOTE: target is on a damped spring so it can continuously update throughout the stride, but without sharp discontinuities
+				Foot.CurrentTargetSpring.Update(Settings.DeltaTime, NewTarget.GetLocation(), 10.0f, 2.0f);
+				NewTarget.SetTranslation(Foot.CurrentTargetSpring.GetCurrent());
+			}
+			else
+			{
+				// the continuously updated target was further away than the previous one, so simply ignore it
+				NewTarget.SetTranslation(Foot.CurrentTargetWorld.GetLocation());
+			}
+		}
+		
+		// set the current target for this foot
+		Foot.CurrentTargetWorld = NewTarget;
+	}
+
+	// run foot collision to prevent overlapping feet targets
+	if (Settings.Stepping.bEnableFootCollision)
+	{
+		// resolve collisions between each foot and every other foot
+		for (FLocomotorFoot* FootPtr : AllFeet)
+		{
+			ResolveFootToFootCollision(*FootPtr);
+		}
+	}
+
+	// last minute cancel footstep if the generated target is too close to the planted foot
+	// NOTE: this can happen if the ground and/or foot collision moved the target
+	// so close to the planted foot that the step would be smaller than some fraction of the minimum threshold step size
+	const float FootStepLengthThreshold = FMath::Pow(Settings.Movement.MinimumStepLength * 0.5f, 2);
+	for (FLocomotorFoot* FootPtr : AllFeet)
+	{
+		FLocomotorFoot& Foot = *FootPtr;
+		if (Foot.bInSwingPhase)
+		{
+			const float StepLengthSq = FVector::DistSquared(Foot.PlantedWorld.GetLocation(), Foot.CurrentTargetWorld.GetLocation());
+			if (StepLengthSq < FootStepLengthThreshold)
+			{
+				// cancel this step
+				Foot.bAtRest = true;
+				Foot.bInSwingPhase = false;
+				Foot.CurrentPhase = Foot.TargetPhase = 1.0f;
+			}
+		}
+	}
+
+	// check if we are fully at rest
+	bFullyAtRest = true;
+	for (const FLocomotorFoot* FootPtr : AllFeet)
+	{
+		if (!FootPtr->bAtRest)
+		{
+			bFullyAtRest = false;
+			break;
+		}
+	}
+}
+
+void FLocomotor::ResolveFootToFootCollision(FLocomotorFoot& Foot)
+{
+	// foot at rest cannot resolve collision
+	if (!Foot.bInSwingPhase)
+	{
+		return;
+	}
+
+	// push target away from other feet
+	for (int32 OtherFootIndex=0; OtherFootIndex < AllFeet.Num(); ++OtherFootIndex)
+	{
+		FLocomotorFoot& OtherFoot = *AllFeet[OtherFootIndex];
+		if (&Foot == &OtherFoot)
+		{
+			continue; // don't collide with self
+		}
+
+		// gather both feet locations and collide them against each other
+		FVector CenterA = Foot.CurrentTargetWorld.GetLocation();
+		FVector CenterB = OtherFoot.CurrentTargetWorld.GetLocation();
+		const float RadiusA = Foot.Settings.CollisionRadius * Settings.Stepping.FootCollisionGlobalScale;
+		const float RadiusB = OtherFoot.Settings.CollisionRadius * Settings.Stepping.FootCollisionGlobalScale;
+		constexpr float InvMassA = 1.0f; // in swing phase
+		const float InvMassB = OtherFoot.bInSwingPhase ? 1.0f : 0.0f;
+		PushCirclesApartInFloorPlane(CenterA, CenterB, RadiusA, RadiusB, InvMassA, InvMassB);
+
+		Foot.CurrentTargetWorld.SetLocation(CenterA);
+		OtherFoot.CurrentTargetWorld.SetLocation(CenterB);
+	}
+}
+
+void FLocomotor::PushCirclesApartInFloorPlane(
+	FVector& CenterA,
+	FVector& CenterB,
+	float RadiusA,
+	float RadiusB,
+	float InvMassA,
+	float InvMassB)
+{
+	// early out if both circles are locked in place
+	const float TotalInvMass = InvMassA + InvMassB;
+	if (TotalInvMass <= KINDA_SMALL_NUMBER)
+	{
+		return;
+	}
+	
+	// project to floor (collision is done in 2d ground plane)
+	FVector2d FloorPositionA(CenterA);
+	FVector2d FloorPositionB(CenterB);
+			
+	// vector from this foot to other foot
+	const FVector2d Delta = FloorPositionB - FloorPositionA;
+	const float Distance = Delta.Size();
+	const float CombinedRadii = RadiusA + RadiusB;
+    
+	// are the circles colliding?
+	if (Distance < CombinedRadii)
+	{
+		const float Overlap = CombinedRadii - Distance;
+		const FVector2d Direction = Delta.GetSafeNormal();
+
+		// calculate the displacement vectors based on the inverse masses
+		const FVector2d CorrectionA = -Direction * (Overlap * (InvMassA / TotalInvMass));
+		const FVector2d CorrectionB = Direction * (Overlap * (InvMassB / TotalInvMass));
+		// update the positions to resolve the collision
+		CenterA = CenterA + FVector(CorrectionA.X, CorrectionA.Y, 0.f);
+		CenterB = CenterB + FVector(CorrectionB.X, CorrectionB.Y, 0.f);
+	}
+}
+
+FVector FLocomotor::ProjectToGroundWithSphereCast(const FLocomotorFoot& Foot, FTransform& TransformToProject)
+{
+	const FVector Offset = FVector(0, 0, FMath::Max(Settings.Stepping.MaxCollisionHeight, 0));
+	const FVector Start = TransformToProject.GetLocation() + Offset;
+	const FVector End = TransformToProject.GetLocation() - Offset;
+	FHitResult HitResult;
+	const bool bDidHit = CastSphere(Start, End, Foot.Settings.CollisionRadius, HitResult);
+	if (bDidHit)
+	{
+		TransformToProject.SetLocation(HitResult.ImpactPoint + FVector(0,0,Foot.InitialRelativeToRootGoal.GetLocation().Z));
+		return HitResult.ImpactNormal;
+	}
+
+	return FVector::ZAxisVector;
+}
+
+void FLocomotor::AnimateFeet()
+{
+	for (FLocomotorFoot* FootPtr : AllFeet)
+	{
+		FLocomotorFoot& Foot = *FootPtr;
+		
+		if (Foot.bAtRest || !Foot.bInSwingPhase)
+		{
+			// foot stays locked at current position while planted or at rest
+			Foot.CurrentWorldFlatPosition = Foot.CurrentWorld.GetLocation();
+			Foot.CurrentWorldFlatPositionNoEase = Foot.CurrentWorld.GetLocation();
+			Foot.PrevWorld = Foot.CurrentWorld;
+			Foot.PlantedWorld = Foot.CurrentWorld;
+			// reset height offset
+			Foot.CurrentHeightOffset = 0.f;
+		}
+
+		// update the current position while in the swing phase
+		if (Foot.bInSwingPhase)
+		{
+			// record the previous world transform before updating the current world transform
+			Foot.PrevWorld = Foot.CurrentWorld;
+			
+			// remap swing range from 0-PhaseWhenSwingEnds to 0-1
+			float PercentStep = FMath::GetRangePct(0.f, Foot.FootPhaseWhenSwingEnds, Foot.CurrentPhase); 
+			// ease in out
+			float PercentStepEased = EaseInOutNorm(PercentStep, Settings.Stepping.StepEaseIn, Settings.Stepping.StepEaseOut);
+			
+			// move foot towards target during swing phase
+			FVector PlantedPosition = Foot.PlantedWorld.GetLocation();
+			FVector TargetPosition = Foot.CurrentTargetWorld.GetLocation();
+			FVector NewPositionEased = FMath::Lerp(PlantedPosition, TargetPosition, PercentStepEased);
+			FVector NewPositionLinear = FMath::Lerp(PlantedPosition, TargetPosition, PercentStep);
+			Foot.CurrentWorldFlatPosition = NewPositionEased;
+			Foot.CurrentWorldFlatPositionNoEase = NewPositionLinear; // record non-eased position as well (prevents jerk in elements reading foot locations)
+			
+			// add height to the current (non-flat) world transform
+			Foot.CurrentHeightOffset = Foot.CurrentStrideHeight * FMath::Sin(PercentStep * PI);
+			const FVector FinalHeightOffset = FVector(0.f, 0.f, Foot.CurrentHeightOffset);
+			Foot.CurrentWorld.SetLocation(Foot.CurrentWorldFlatPosition + FinalHeightOffset);
+		}
+
+		// generate an animated heel peel rotation
+		FQuat PeelHeelRotation = FQuat::Identity;
+		if (Foot.bInSwingPhase)
+		{
+			// add heel peel to current (non-flat) world rotation
+			constexpr float PeelDurationAsPercentOfSwing = 0.2;
+			constexpr float PeelPhaseBufferBeforeFootFall = 0.1f;
+			constexpr float PhaseAtPeelStart = 0.0f;
+			const float PhaseAtMaxPeel = Foot.FootPhaseWhenSwingEnds * PeelDurationAsPercentOfSwing;
+			const float PhaseAtPeelEnd = Foot.FootPhaseWhenSwingEnds - PeelPhaseBufferBeforeFootFall;
+			// scale the maximum peel angle by the stride length
+			FQuat MaxPeelRotation = FQuat::MakeFromEuler(Foot.Settings.MaxHeelPeel);
+			FQuat MinPeelRotation = FQuat::FastLerp(FQuat::Identity, MaxPeelRotation, 0.25f);
+			PeelHeelRotation = FQuat::FastLerp(MinPeelRotation, MaxPeelRotation, CurrentPercentOfMaxSpeed);
+			// rotate to max peel at start of phase and then back to flat foot before foot-fall
+			if (Foot.CurrentPhase > PhaseAtMaxPeel)
+			{
+				// blend back to flat foot
+				const float PercentFromMaxPeelToEndOfStride =  FMath::GetRangePct(PhaseAtMaxPeel, PhaseAtPeelEnd, Foot.CurrentPhase); 
+				PeelHeelRotation = FQuat::FastLerp(PeelHeelRotation, FQuat::Identity, PercentFromMaxPeelToEndOfStride).GetNormalized();
+			}
+			else
+			{
+				// blend from flat foot to peeled heel
+				const float PercentToMaxPeel =  FMath::GetRangePct(PhaseAtPeelStart, PhaseAtMaxPeel, Foot.CurrentPhase);
+				PeelHeelRotation = FQuat::FastLerp(FQuat::Identity, PeelHeelRotation, PercentToMaxPeel).GetNormalized();
+			}
+		}
+
+		// update current rotation through a quat spring
+		// this is done regardless of what phase the foot is in (we allow foot pivoting while planted)
+		{
+			constexpr float RotationStiffness = 40.f;
+			constexpr float RotationDamping = .9f;
+			
+			const FQuat NewFootRotation = Foot.CurrentRotationSpring.Update(
+				Settings.DeltaTime,
+				Foot.CurrentTargetWorld.GetRotation(),
+				RotationStiffness,
+				RotationDamping);
+			Foot.CurrentWorld.SetRotation(NewFootRotation * PeelHeelRotation);
+		}
+	}
+
+	// DEBUG OUTPUT
+	//const FLocomotorFoot& Foot = *AllFeet[3];
+	//float PercentSwing = FMath::GetRangePct(0.f, Foot.FootPhaseWhenSwingEnds, Foot.CurrentPhase);
+	//float PercentSwingBasedOnFoot = FMath::GetRangePct(0.0f, Foot.FootPhaseWhenSwingEnds, Foot.CurrentPhase);
+	//UE_LOG(LogTemp, Warning, TEXT("CurrentPhase:            %f"), Foot.CurrentPhase);
+	//UE_LOG(LogTemp, Warning, TEXT("PercentSwing:            %f"), PercentSwing);
+	//UE_LOG(LogTemp, Warning, TEXT("PercentSwingBasedOnFoot: %f"), PercentSwingBasedOnFoot);
+	//UE_LOG(LogTemp, Warning, TEXT("PhaseWhenSwingEnds:      %f"), Foot.FootPhaseWhenSwingEnds);
+	//UE_LOG(LogTemp, Warning, TEXT("InSwingPhase:            %s"), Foot.bInSwingPhase ? TEXT("True") : TEXT("False"));
+}
+
+void FLocomotor::UpdateBody()
+{
+	if (AllFeet.IsEmpty())
+	{
+		// avoid divide by zero
+		return;
+	}
+	
+	// update body based on current feet
+	{
+		Pelvis.CalculateBodyTransform(AllFeet, Settings.CurrentWorldRootGoal, Pelvis.CurrentBodyTargetWorld);	
+	}
+
+	// pitch/roll the body target based on the slope of the ground
+	const bool bPitchBodyWithGround = !FMath::IsNearlyZero(Settings.Pelvis.OrientToGroundPitch);
+	const bool bRollBodyWithGround = !FMath::IsNearlyZero(Settings.Pelvis.OrientToGroundRoll);
+	if (bPitchBodyWithGround || bRollBodyWithGround)
+	{
+		// get the average normal under the feet
+		FVector Normal = FVector::ZeroVector;
+		for (FLocomotorFoot* Foot : AllFeet)
+		{
+			const FVector Offset = FVector(0, 0, FMath::Max(Settings.Stepping.MaxCollisionHeight, 0));
+			FVector Start = Foot->CurrentWorld.GetLocation() + Offset;
+			FVector End = Foot->CurrentWorld.GetLocation() - Offset;
+			FHitResult HitResult;
+			const bool bDidHit = CastSphere(Start, End, Foot->Settings.CollisionRadius, HitResult);
+			if (bDidHit)
+			{
+				Normal += HitResult.Normal;
+			}
+		}
+		Normal.Normalize();
+
+		// decompose normal into pitch and roll rotations
+		const FVector Fwd = Settings.CurrentWorldRootGoal.GetRotation().RotateVector(FVector(0,1,0));
+		const FVector Up = FVector(0,0,1);
+		FVector Side;
+		FQuat PitchRotation;
+		FQuat RollRotation;
+		CalcPitchRollFromNormal(Fwd, Up, Normal, Side,PitchRotation, RollRotation);
+
+		// scale rotations by user amount
+		PitchRotation = FQuat::FastLerp(FQuat::Identity, PitchRotation, Settings.Pelvis.OrientToGroundPitch);
+		RollRotation = FQuat::FastLerp(FQuat::Identity, RollRotation, Settings.Pelvis.OrientToGroundRoll);
+
+		// apply to the current target
+		FQuat NewPelvisRotation = PitchRotation * RollRotation * Pelvis.CurrentBodyTargetWorld.GetRotation();
+		Pelvis.CurrentBodyTargetWorld.SetRotation(NewPelvisRotation);
+	}
+
+	// update current body based on target + lead
+	{
+		// find the current velocity of the pelvis target, and use that to extrapolate the pelvis so that it leads the feet motion
+		// the amount the pelvis leads the foot motion is controlled with LeadAmount.
+		// the lead itself is ran through a spring damper so that the pelvis smoothly leads and returns to the target location
+		const FVector TargetVelocity = Pelvis.CurrentBodyTargetWorld.GetLocation() - Pelvis.PreviousBodyTargetWorld.GetLocation();
+		float LeadAmountToUse = bAccelerating ? Settings.Pelvis.LeadAmount : -Settings.Pelvis.LeadAmount;
+		const FVector LeadTarget = TargetVelocity * LeadAmountToUse;
+		Pelvis.CurrentBodyLead = Pelvis.BodyLeadDamper.Update(
+			LeadTarget,
+			Settings.DeltaTime,
+			Settings.Pelvis.LeadDampingHalfLife);
+		const FVector ExtrapolatedBodyPosition = Pelvis.CurrentBodyTargetWorld.GetLocation() + Pelvis.CurrentBodyLead;
+		Pelvis.CurrentBodyWorld.SetLocation(ExtrapolatedBodyPosition);	
+	}
+
+	// update current body rotation through a quat spring
+	{
+		const FQuat NewBodyRotation = Pelvis.BodyRotationSpring.Update(
+			Settings.DeltaTime,
+			Pelvis.CurrentBodyTargetWorld.GetRotation(),
+			Settings.Pelvis.RotationStiffness,
+			Settings.Pelvis.RotationDamping);
+		Pelvis.CurrentBodyWorld.SetRotation(NewBodyRotation);
+	}
+
+	// record prev target for next frame
+	Pelvis.PreviousBodyTargetWorld = Pelvis.CurrentBodyTargetWorld;
+}
+
+void FLocomotor::AnimatePelvis()
+{
+	// update pelvis based on body motion
+	{
+		// THIS BROKE INPUT POSE SUPPORT, BUT THAT WAS FLIPPING ON SOME ASSETS :(
+		const FTransform InitialPelvisRelativeToBody = Pelvis.InitialWorld.GetRelativeTransform(Pelvis.InitialBodyWorld);
+		const FTransform CurrentWorldPelvisTarget = InitialPelvisRelativeToBody * Pelvis.CurrentBodyWorld;
+		
+		/*
+		// generate pelvis transform based on current body transform
+		const FTransform InitialRootGoalWorldNoRotation = FTransform(FQuat::Identity, InitialRootGoalWorld.GetLocation());
+		const FTransform InitialBodyComponentSpace = Pelvis.InitialBodyWorld.GetRelativeTransform(InitialRootGoalWorldNoRotation);
+		// use continuously updated pelvis transform (InputPelvisLocal) to support input animated pelvis 
+		const FTransform InputPelvisRelativeToInitialBody = Settings.Pelvis.InputPelvisComponentSpace.GetRelativeTransform(InitialBodyComponentSpace);
+		const FTransform CurrentWorldPelvisTarget = InputPelvisRelativeToInitialBody * Pelvis.CurrentBodyWorld;
+		*/
+		
+		
+		/*
+		// generate pelvis transform based on current body transform
+		const FTransform InitialPelvisRelativeToBody = Pelvis.InitialWorld.GetRelativeTransform(Pelvis.InitialBodyWorld);
+		Pelvis.CurrentWorld = InitialPelvisRelativeToBody * Pelvis.CurrentBodyWorld;
+		*/
+		
+		/*
+		// generate pelvis transform based on current body transform
+		const FTransform InitialBodyLocal = Pelvis.InitialBodyWorld.GetRelativeTransform(InitialRootGoalWorld);
+		// use continuously updated pelvis transform (InputPelvisLocal) to support input animated pelvis 
+		const FTransform InputPelvisRelativeToInitialBody = Settings.Pelvis.InputPelvisLocal.GetRelativeTransform(InitialBodyLocal);
+		const FTransform CurrentWorldPelvisTarget = InputPelvisRelativeToInitialBody * Pelvis.CurrentBodyWorld;
+		*/
+
+		// run it through a damper
+		//const FVector DampedPelvisPosition = Pelvis.PelvisPositionDamper.Update(
+		//	CurrentWorldPelvisTarget.GetLocation(),
+		//	Settings.DeltaTime,
+		//	Settings.Pelvis.PositionDampingHalfLife);
+		//Pelvis.CurrentWorld.SetLocation(DampedPelvisPosition + Pelvis.CurrentBodyLead);
+		
+		Pelvis.CurrentWorld = CurrentWorldPelvisTarget;
+	}
+
+	// additively apply vertical pelvis bob based on phase
+	{
+		// calculate average amount feet have been raised
+		float TargetBobOffset = 0.f;
+		for (FLocomotorFoot* FootPtr : AllFeet)
+		{
+			TargetBobOffset += FootPtr->CurrentHeightOffset;
+		}
+		TargetBobOffset /= AllFeet.IsEmpty() ? 1.0f : AllFeet.Num();
+
+		// add a static offset when moving
+		// remove static offset when foot is at rest, this removes the "double bent knee slide" when going to rest
+		float NumFeetAtRest = 0.f;
+		for (const FLocomotorFoot* Foot : AllFeet)
+		{
+			NumFeetAtRest += Foot->bAtRest ? 1.0f : 0.f;
+		}
+		const float PercentFeetAtRest = 1.0f - NumFeetAtRest / AllFeet.Num();
+		const float ScaledStaticOffset = Settings.Pelvis.BobOffset * PercentFeetAtRest;
+		TargetBobOffset += ScaledStaticOffset;
+	
+		// put height offset through a spring
+		const float SmoothBobOffset = Pelvis.BobOffsetSpring.Update(
+			Settings.DeltaTime,
+			TargetBobOffset,
+			Settings.Pelvis.BobStiffness,
+			Settings.Pelvis.BobDamping);
+		
+		// apply height to pelvis
+		// current location is a flat interpolation, so we can add an offset to it
+		FVector CurrentLocation = Pelvis.CurrentWorld.GetLocation();
+		CurrentLocation.Z += SmoothBobOffset;
+		Pelvis.CurrentWorld.SetTranslation(CurrentLocation);
+	}
+}
+
+void FLocomotor::AnimateHead()
+{
+	if (!Head.bInitialized)
+	{
+		return;
+	}
+
+	//Head.CurrentWorld.
+}
+
+float FLocomotor::EaseInOutNorm(const float Input, float EaseInAmount, float EaseOutAmount) const
+{
+	const float InputSquared = Input * Input;
+	const float EasedInput =  InputSquared / (2.0f * (InputSquared - Input) + 1.0f);
+	const float ClampedEaseAmount = FMath::Clamp(Input < 0.5 ? EaseInAmount : EaseOutAmount, 0.0f, 1.0f);
+	return FMath::Lerp(Input, EasedInput, ClampedEaseAmount);
+}
+
+float FLocomotor::WrapPhaseInRange(float PhaseToWrap)
+{
+	ensure(PhaseToWrap >= -1.f);
+	if (PhaseToWrap < 0.0f)
+	{
+		PhaseToWrap += 1.0f;
+	}
+	return PhaseToWrap - FMath::Floor(PhaseToWrap);
+}
+
+float FLocomotor::BlendTowardsTargetPhase(const float CurrentPhase, const float TargetPhase, const float PhaseBlendRate)
+{
+	// calculate distance to target in the positive direction
+	float DirectDistance = TargetPhase - CurrentPhase;
+	if (DirectDistance < 0.0f)
+	{
+		// if directDistance is negative, it means wrapping around is shorter
+		DirectDistance += 1.0f;
+	}
+
+	// if the distance is greater than a 1/4 phase, pause the current phase until the target phase gets closer
+	constexpr float PHASE_OFFSET_THRESHOLD = 0.25f;
+	if (DirectDistance > PHASE_OFFSET_THRESHOLD)
+	{
+		return CurrentPhase;
+	}
+	
+	// calculate the distance to move, this ensures we do not overshoot the target value
+	const float DistanceToMove = FMath::Min(DirectDistance, PhaseBlendRate);
+
+	// increment and wrap
+	return WrapPhaseInRange(CurrentPhase + DistanceToMove);
+}
+
+float FLocomotor::GetPhaseOffsetForSetFromMovementStyle(int32 SetIndex)
+{
+	if (Settings.Movement.Styles.IsEmpty())
+	{
+		return 0.f;
+	}
+	
+	if (!ensure(FootSets.IsValidIndex(SetIndex)))
+	{
+		return 0.f;
+	}
+
+	if (Settings.Movement.Styles.Num() == 1)
+	{
+		return 0.f;
+	}
+
+	//TODO
+	return 0.f;
+}
+
+bool FLocomotor::CastSphere(const FVector& Start, const FVector& End, float SphereRadius, FHitResult& OutHitResult)
+{
+	if (!Settings.Stepping.World)
+	{
+		return false;
+	}
+	
+	const ECollisionChannel CollisionChannel = UEngineTypes::ConvertToCollisionChannel(Settings.Stepping.TraceChannel); 
+	const FCollisionShape CollisionShape = FCollisionShape::MakeSphere(SphereRadius);
+	
+	FCollisionQueryParams QueryParams;
+	QueryParams.bTraceComplex = true;
+	if (Settings.Stepping.IgnoredActor)
+	{
+		QueryParams.AddIgnoredActor(Settings.Stepping.IgnoredActor);
+	}
+	else if (const UPrimitiveComponent* PrimitiveComponent = Cast<UPrimitiveComponent>(Settings.Stepping.IgnoredComponent))
+	{
+		QueryParams.AddIgnoredComponent(PrimitiveComponent);
+	}
+	
+	FCollisionResponseParams ResponseParams(ECR_Ignore);
+	ResponseParams.CollisionResponse.SetResponse(ECC_WorldStatic, ECR_Block);
+	
+	bool bHit = Settings.Stepping.World->SweepSingleByChannel(
+		OutHitResult,
+		Start,
+		End,
+		FQuat::Identity,
+		CollisionChannel,
+		CollisionShape,
+		QueryParams,
+		ResponseParams);
+
+	return bHit;
+}
+
+void FLocomotor::CalcPitchRollFromNormal(
+	const FVector& Fwd,
+	const FVector& Up,
+	const FVector& Normal,
+	FVector& OutSide,
+	FQuat& OutPitchRotation,
+	FQuat& OutRollRotation)
+{
+	// decompose normal into PITCH (forward/back) and ROLL (side to side) relative to given UP and FWD vectors
+	
+	OutSide = Up.Cross(Fwd);
+
+	// pitch (in forward / up plane)
+	const FVector NormalInFwdUpPlane = FVector::VectorPlaneProject(Normal, OutSide).GetSafeNormal();
+	OutPitchRotation = FQuat::FindBetweenNormals(Up, NormalInFwdUpPlane);
+	// roll (in side / up plane)
+	const FVector NormalInSideUpPlane = FVector::VectorPlaneProject(Normal, Fwd).GetSafeNormal();
+	OutRollRotation = FQuat::FindBetweenNormals(Up, NormalInSideUpPlane);
+};
+
+float FLocomotor::GetPhaseCurrent() const
+{
+	return CurrentGlobalPhase;
+}
diff --git a/Engine/Plugins/Experimental/Animation/Locomotor/Source/Locomotor/Private/RigUnit_Locomotor.cpp b/Engine/Plugins/Experimental/Animation/Locomotor/Source/Locomotor/Private/RigUnit_Locomotor.cpp
new file mode 100644
--- /dev/null
+++ b/Engine/Plugins/Experimental/Animation/Locomotor/Source/Locomotor/Private/RigUnit_Locomotor.cpp
@@ -0,0 +1,287 @@
+﻿// Copyright Epic Games, Inc. All Rights Reserved.
+
+#include "RigUnit_Locomotor.h"
+
+
+#include "Kismet/KismetMathLibrary.h"
+#include "Units/RigUnitContext.h"
+#include "Math/ControlRigMathLibrary.h"
+
+FRigUnit_Locomotor_Execute()
+{
+	DECLARE_SCOPE_HIERARCHICAL_COUNTER_RIGUNIT()
+
+	// awaiting hierarchy
+	URigHierarchy* Hierarchy = ExecuteContext.Hierarchy;
+	if (Hierarchy == nullptr)
+	{
+		return;
+	}
+
+	// awaiting root control
+	if (!RootControl.IsValid())
+	{
+		return;
+	}
+
+	// cache index of root control
+	CachedRootElement.UpdateCache(FRigElementKey(RootControl, ERigElementType::Control), Hierarchy);
+
+	// check if we need to reinitialize
+	bool bNeedsReinitialized = !Locomotor.HasFeet();
+	if (FeetTransforms.Num() != Locomotor.GetFeet().Num())
+	{
+		bNeedsReinitialized = true;
+	}
+
+	FTransform ToWorld = ExecuteContext.GetToWorldSpaceTransform();
+	
+	// initialize the simulation
+	if (bNeedsReinitialized)
+	{	
+		// initialize with the initial root goal and initial pelvis transform
+		const FTransform InitialRootGoalGlobal = Hierarchy->GetInitialGlobalTransform(CachedRootElement.GetIndex());
+		const FTransform InitialPelvisGlobal = Hierarchy->GetInitialGlobalTransform(Pelvis.PelvisBone);
+		const FTransform InitialPelvisWorld = InitialPelvisGlobal * ToWorld;
+		const FTransform InitialRootGoalWorld = InitialRootGoalGlobal.GetRelativeTransform(InitialPelvisGlobal) * InitialPelvisWorld;
+		Locomotor.Reset(InitialRootGoalWorld, InitialPelvisWorld);
+
+		// add all the feet
+		for (const FFootSet& FootSet : FootSets)
+		{
+			// first create a set for these feet
+			int32 FootSetIndex = Locomotor.AddFootSet(FootSet.PhaseOffset);
+			// add all the feet in the set
+			for (const FFootSettings& Foot : FootSet.Feet)
+			{
+				if (Foot.AnkleBone.IsValid())
+				{
+					// calculate initial foot in world space
+					FTransform InitialFootTransform = Hierarchy->GetInitialGlobalTransform(Foot.AnkleBone);
+					FVector InitialFootTransformOffset = InitialFootTransform.GetRotation().RotateVector(Foot.StaticLocalOffset);
+					InitialFootTransform.SetTranslation(InitialFootTransform.GetLocation() + InitialFootTransformOffset);
+					const FTransform InitialWorldFoot = InitialFootTransform * ToWorld;
+
+					// assemble per-foot settings
+					FLocomotorFootSettings FootSettings;
+					FootSettings.CollisionRadius = Foot.CollisionRadius;
+					FootSettings.MaxHeelPeel = Foot.MaxHeelPeel;
+					FootSettings.StaticLocalOffset = Foot.StaticLocalOffset	* ToWorld.GetScale3D();
+					
+					Locomotor.AddFootToSet(FootSetIndex, InitialWorldFoot, FootSettings);
+				}	
+			}
+		}
+
+		// create room for output
+		FeetTransforms.SetNumZeroed(Locomotor.GetFeet().Num());
+	}
+
+	if (!Locomotor.HasFeet())
+	{
+		return;
+	}
+
+	FTransform RootTransformNoScale = Hierarchy->GetGlobalTransform(CachedRootElement.GetIndex());
+	RootTransformNoScale.SetScale3D(FVector::One());
+	
+	// update the simulation settings
+	FLocomotorInputSettings Settings;
+	Settings.CurrentWorldRootGoal = RootTransformNoScale;
+	Settings.DeltaTime = ExecuteContext.GetDeltaTime() * Movement.GlobalTimeScale;
+	// movement
+	Settings.Movement.MinimumStepLength = Movement.MinimumStepLength;
+	Settings.Movement.SpeedMax = Movement.SpeedMax;
+	Settings.Movement.SpeedMin = Movement.SpeedMin;
+	Settings.Movement.PhaseSpeedMax = Movement.PhaseSpeedMax;
+	Settings.Movement.PhaseSpeedMin = Movement.PhaseSpeedMin;
+	Settings.Movement.Acceleration = Movement.Acceleration;
+	Settings.Movement.Deceleration = Movement.Deceleration;
+	for (ELocomotorMovementStyle Style : Movement.Styles)
+	{
+		EMovementStyle StyleToAdd = static_cast<EMovementStyle>(Style); // cast from BP enum to internal enum
+		Settings.Movement.Styles.Add(StyleToAdd);
+	}
+	Settings.Stepping.PercentOfStrideInAir = Stepping.PercentOfStrideInAir;
+	Settings.Stepping.AirExtensionAtMaxSpeed = Stepping.AirExtensionAtMaxSpeed;
+	Settings.Stepping.StepHeight = Stepping.StepHeight;
+	Settings.Stepping.StepEaseIn = Stepping.StepEaseIn;
+	Settings.Stepping.StepEaseOut = Stepping.StepEaseOut;
+	Settings.Stepping.bEnableFootCollision = Stepping.bEnableFootCollision;
+	Settings.Stepping.FootCollisionGlobalScale = Stepping.FootCollisionGlobalScale;
+	Settings.Stepping.bEnableGroundCollision = Stepping.bEnableGroundCollision;
+	Settings.Stepping.MaxCollisionHeight = Stepping.MaxCollisionHeight;
+	Settings.Stepping.TraceChannel = Stepping.TraceChannel;
+	Settings.Stepping.OrientFootToGroundPitch = Stepping.OrientFootToGroundPitch;
+	Settings.Stepping.OrientFootToGroundRoll = Stepping.OrientFootToGroundRoll;
+	Settings.Stepping.IgnoredActor = ExecuteContext.GetOwningActor();
+	Settings.Stepping.IgnoredComponent = ExecuteContext.GetOwningComponent();
+	Settings.Stepping.World = ExecuteContext.GetWorld();
+	// pelvis
+	Settings.Pelvis.InputPelvisComponentSpace = Hierarchy->GetGlobalTransform(Pelvis.PelvisBone);
+	Settings.Pelvis.PositionDampingHalfLife = Pelvis.PositionDampingHalfLife;
+	Settings.Pelvis.RotationStiffness = Pelvis.RotationStiffness;
+	Settings.Pelvis.RotationDamping = Pelvis.RotationDamping;
+	Settings.Pelvis.LeadAmount = Pelvis.LeadAmount;
+	Settings.Pelvis.LeadDampingHalfLife = Pelvis.LeadDampingHalfLife;
+	Settings.Pelvis.BobOffset = Pelvis.BobOffset;
+	Settings.Pelvis.BobStiffness = Pelvis.BobStiffness;
+	Settings.Pelvis.BobDamping = Pelvis.BobDamping;
+	Settings.Pelvis.OrientToGroundPitch = Pelvis.OrientToGroundPitch;
+	Settings.Pelvis.OrientToGroundRoll = Pelvis.OrientToGroundRoll;
+	// tick the loco motor
+	Locomotor.RunSimulation(Settings);
+	
+	// output the resulting foot transforms
+	const TArray<FLocomotorFoot*>& AllFeet = Locomotor.GetFeet();
+	for (int32 FootIndex=0 ; FootIndex<AllFeet.Num(); ++FootIndex)
+	{
+		FeetTransforms[FootIndex] = ExecuteContext.ToVMSpace(AllFeet[FootIndex]->CurrentWorld);
+	}
+
+	// transform the pelvis bone (and propagate to all children)
+	constexpr bool bInitial = false;
+	constexpr bool bAffectChildren = true;
+	constexpr bool bSetupUndo = false;
+	constexpr bool bPrintPythonCommands = false;
+	Hierarchy->SetGlobalTransform(
+		Pelvis.PelvisBone,
+		ExecuteContext.ToVMSpace(Locomotor.GetPelvisCurrent()),
+		bInitial,
+		bAffectChildren,
+		bSetupUndo,
+		bPrintPythonCommands);
+	
+	// do all debug drawing
+	Debug.DrawDebug(ExecuteContext.GetDrawInterface(), Locomotor, ExecuteContext.GetToWorldSpaceTransform().Inverse());
+}
+
+void FLocomotorDebugSettings::DrawDebug(
+	FRigVMDrawInterface* DrawInterface,
+	const FLocomotor& Locomotor,
+	const FTransform& WorldToGlobal) const
+{
+	{
+		if (DrawInterface == nullptr || !bDrawDebug)
+		{
+			return;
+		}
+
+		const TArray<FLocomotorFoot*>& Feet = Locomotor.GetFeet();
+		int32 FootIndex = 0;
+		for (const FLocomotorFoot* Foot : Feet)
+		{
+			const float FootPointSize = Scale * 0.35f;
+			
+			if (bDrawCurrentFeet)
+			{
+				// draw current transform
+				DrawInterface->DrawAxes(WorldToGlobal, Foot->CurrentWorld, FootPointSize, Thickness);
+			}
+
+			if (bDrawCurrentFeetFlat)
+			{
+				// draw current flat transform
+				DrawInterface->DrawPoint(WorldToGlobal, Foot->CurrentWorldFlatPosition, FootPointSize, FLinearColor::Blue);
+			}
+
+			if (bDrawCurrentFeetTarget)
+			{
+				// draw current flat target
+				DrawInterface->DrawPoint(WorldToGlobal, Foot->CurrentTargetWorld.GetLocation(), FootPointSize, FLinearColor::Green);
+			}
+
+			if (bDrawPrevFeetTarget)
+			{
+				// draw prev target transform
+				DrawInterface->DrawPoint(WorldToGlobal, Foot->PlantedWorld.GetLocation(), FootPointSize, FLinearColor::Yellow);
+				DrawInterface->DrawLine(
+					WorldToGlobal,
+					Foot->PlantedWorld.GetLocation(),
+					Foot->CurrentTargetWorld.GetLocation(),
+					FLinearColor::Yellow, Thickness);
+				DrawInterface->DrawLine(
+					WorldToGlobal,
+					Foot->CurrentTargetWorld.GetLocation(),
+					Foot->FinalTargetWorld.GetLocation(),
+					FLinearColor::Gray, Thickness);
+			}
+			
+			if (bDrawFinalFeetTarget)
+			{
+				// draw final target transform
+				DrawInterface->DrawPoint(WorldToGlobal, Foot->FinalTargetWorld.GetLocation(), FootPointSize, FLinearColor::Black);
+			}
+
+			if (bDrawFeetCollision)
+			{
+				FTransform FootNoRotation = Foot->CurrentWorld;
+				FootNoRotation.SetRotation(FQuat::Identity);
+				DrawInterface->DrawCircle(
+					WorldToGlobal,
+					FootNoRotation,
+					Foot->Settings.CollisionRadius * Locomotor.GetSettings().Stepping.FootCollisionGlobalScale,
+					FLinearColor::White,
+					Thickness,
+					12);
+			}
+
+			++FootIndex;
+		}
+
+		// draw the phase circle with arrow that rotates around with the phase
+		if (bDrawPhaseCircle)
+		{
+			const float Radius = Scale * 2.0f;
+			
+			// draw goal circle
+			const FLocomotorInputSettings& Settings = Locomotor.GetSettings();
+			FTransform RootGoalNoRotation = Settings.CurrentWorldRootGoal * WorldToGlobal;
+			RootGoalNoRotation.SetRotation(FQuat::Identity);
+			DrawInterface->DrawCircle(
+				FTransform::Identity,
+				RootGoalNoRotation,
+				Radius,
+				FLinearColor::Black,
+				Thickness*2.f,
+				24 /*num sides*/);
+
+			// draw phase line
+			const float AngleRadians = FMath::Lerp(0.0f, 2*PI, Locomotor.GetPhaseCurrent());
+			const FVector ArrowDirection = FVector(FMath::Cos(AngleRadians), FMath::Sin(AngleRadians), 0.0f) * Radius;
+			DrawInterface->DrawLine(
+					FTransform::Identity,
+					RootGoalNoRotation.GetLocation(),
+					RootGoalNoRotation.GetLocation() + ArrowDirection,
+					FLinearColor::Green, Thickness);
+
+			// draw forward arrow
+			const FVector ArrowSide = FVector(0.0f, 3.0f, 0.0f);
+			DrawInterface->DrawArrow(
+				Settings.CurrentWorldRootGoal * WorldToGlobal,
+				FVector::RightVector * Radius,
+				ArrowSide,
+				FLinearColor::Black,
+				Thickness*2.f);
+		}
+
+		if (bDrawBody)
+		{
+			FTransform BodyCurrent = Locomotor.GetBodyCurrent();
+			FTransform BodyTarget = Locomotor.GetBodyTarget();
+			FTransform PelvisCurrent = Locomotor.GetPelvisCurrent();
+			BodyCurrent.SetScale3D(BodyCurrent.GetScale3D() * Scale * 0.4f);
+			BodyTarget.SetScale3D(BodyTarget.GetScale3D() * Scale * 0.2f);
+			PelvisCurrent.SetScale3D(PelvisCurrent.GetScale3D() * Scale * 0.5f);
+			DrawInterface->DrawBox(FTransform::Identity,BodyCurrent, FLinearColor::Green, Thickness);
+			DrawInterface->DrawBox(FTransform::Identity,BodyTarget, FLinearColor::Black, Thickness);
+			DrawInterface->DrawBox(FTransform::Identity, PelvisCurrent, FLinearColor::Green, Thickness);
+
+			DrawInterface->DrawLine(
+					FTransform::Identity,
+					BodyCurrent.GetLocation(),
+					PelvisCurrent.GetLocation(),
+					FLinearColor::Gray, Thickness);
+		}
+	}
+}
diff --git a/Engine/Plugins/Experimental/Animation/Locomotor/Source/Locomotor/Public/Locomotor.h b/Engine/Plugins/Experimental/Animation/Locomotor/Source/Locomotor/Public/Locomotor.h
new file mode 100644
--- /dev/null
+++ b/Engine/Plugins/Experimental/Animation/Locomotor/Source/Locomotor/Public/Locomotor.h
@@ -0,0 +1,14 @@
+// Copyright Epic Games, Inc. All Rights Reserved.
+
+#pragma once
+
+#include "Modules/ModuleManager.h"
+
+class FLocomotorModule : public IModuleInterface
+{
+public:
+
+	/** IModuleInterface implementation */
+	virtual void StartupModule() override;
+	virtual void ShutdownModule() override;
+};
diff --git a/Engine/Plugins/Experimental/Animation/Locomotor/Source/Locomotor/Public/LocomotorCore.h b/Engine/Plugins/Experimental/Animation/Locomotor/Source/Locomotor/Public/LocomotorCore.h
new file mode 100644
--- /dev/null
+++ b/Engine/Plugins/Experimental/Animation/Locomotor/Source/Locomotor/Public/LocomotorCore.h
@@ -0,0 +1,549 @@
+﻿// Copyright Epic Games, Inc. All Rights Reserved.
+
+#pragma once
+
+#include "Engine/EngineTypes.h"
+#include "Engine/HitResult.h"
+#include "Kismet/KismetMathLibrary.h"
+#include "Math/Transform.h"
+#include "LocomotorCore.generated.h"
+
+// Daniel Holden's "Exact Damper" described here https://theorangeduck.com/page/spring-roll-call
+struct FVectorDamper
+{
+	const FVector& Update(const FVector& TargetValue, const float DeltaTime, const float HalfLife=0.1f)
+	{
+		constexpr float NaturalLogTwo = 0.69314718056f;
+		constexpr float Epsilon = 1e-5f;
+		const float Param = 1.0f - FastNegExp((NaturalLogTwo * DeltaTime) / (HalfLife + Epsilon));
+		CurrentValue.X = FMath::Lerp(CurrentValue.X, TargetValue.X, Param);
+		CurrentValue.Y = FMath::Lerp(CurrentValue.Y, TargetValue.Y, Param);
+		CurrentValue.Z = FMath::Lerp(CurrentValue.Z, TargetValue.Z, Param);
+		return CurrentValue;
+	}
+	
+	static float FastNegExp(float x)
+	{
+		// fast approximation of negative exponent
+		return 1.0f / (1.0f + x + 0.48f*x*x + 0.235f*x*x*x);
+	}
+
+	void Reset(const FVector& InValue)
+	{
+		CurrentValue = InValue;
+	}
+
+	FVector GetCurrentValue() const
+	{
+		return CurrentValue;
+	};
+	
+private:
+
+	FVector CurrentValue;
+};
+
+struct FFloatSpring
+{
+	FFloatSpring() = default;
+
+	float Update(float DeltaTime, const float& TargetValue, const float Stiffness, const float Damping)
+	{
+		CurrentValue = UKismetMathLibrary::FloatSpringInterp(
+		CurrentValue, TargetValue, SpringState,
+		Stiffness,
+		Damping,
+		DeltaTime, 1.0f, 0.0f);
+		return CurrentValue;
+	}
+
+	void Reset(const float& InValue)
+	{
+		CurrentValue = InValue;
+		SpringState.Reset();
+	}
+
+	float GetCurrent() const
+	{
+		return CurrentValue;
+	}
+
+private:
+
+	float CurrentValue;
+	FFloatSpringState SpringState;
+};
+
+struct FVectorSpring
+{
+	FVectorSpring() = default;
+
+	FVector Update(float DeltaTime, const FVector& TargetValue, const float Stiffness, const float Damping)
+	{
+		CurrentValue = UKismetMathLibrary::VectorSpringInterp(
+		CurrentValue, TargetValue, SpringState,
+		Stiffness,
+		Damping,
+		DeltaTime, 1.0f, 0.0f);
+		return CurrentValue;
+	}
+
+	void Reset(const FVector& InValue)
+	{
+		CurrentValue = InValue;
+		SpringState.Reset();
+	}
+
+	FVector GetCurrent() const
+	{
+		return CurrentValue;
+	}
+
+private:
+
+	FVector CurrentValue;
+	FVectorSpringState SpringState;
+};
+
+struct FQuatSpring
+{
+	
+	FQuatSpring() = default;
+
+	FQuat Update(float DeltaTime, const FQuat& Target, const float Stiffness, const float Damping)
+	{
+		Current = UKismetMathLibrary::QuaternionSpringInterp(
+			Current, Target, SpringState,
+			Stiffness, Damping, DeltaTime,
+			1.0f, 0.0f);
+		return Current;
+	}
+
+	void Reset(const FQuat& InValue)
+	{
+		Current = InValue;
+		SpringState.Reset();
+	}
+
+	FQuat GetCurrent() const
+	{
+		return Current;
+	}
+
+private:
+
+	FQuat Current;
+	FQuaternionSpringState SpringState;
+};
+
+enum class EMovementStyle : uint8
+{
+	Walk,
+	Trot,
+	Gallop,
+};
+
+struct FLocomotorMovementSettings
+{
+	// the smallest step to make
+	float MinimumStepLength = 10.0f;
+	// the maximum movement speed in cm/s
+	float SpeedMax = 80.f;
+	// the maximum movement speed in cm/s
+	float SpeedMin = 50.f;
+	// the maximum phase speed in cycles/s
+	float PhaseSpeedMax = 4.f;
+	// the minimum phase speed in cycles/s
+	float PhaseSpeedMin = 1.0f;
+	// the acceleration in cm/s/s
+	float Acceleration = 100.f;
+	// the acceleration in cm/s/s
+	float Deceleration  = 30.f;
+	// movement styles
+	TArray<EMovementStyle> Styles;
+
+	//static TMap<EMovementStyle, float> StyleOffsets;
+};
+
+struct FLocomotorStepSettings
+{
+	// Default is 0.35. Range is 0.1 - 0.9. The percentage amount of the 0-1 stride phase where the foot is in the air.
+	// At the default (0.35), the foot will plant on the ground 35% through it's total phase.
+	// At lower values, the foot will spend more time on the ground as a percentage of the total stride time.
+	// At higher values, each foot will spend less time on the ground as a percentage of the total stride time.
+	float PercentOfStrideInAir = 0.35f;
+	// Default is 0.2. Range is 0.0 - 0.9. This value is added to the "PercentOfStrideInAir" to extend the air-time while at max speed.
+	// For example, if PercentOfStrideInAir is at 0.5, and the air extension is at 0.4 the foot will be in the air for 90% of the stride cycle.
+	// This value is gradually blended on as the Speed accelerates between SpeedMin and SpeedMax.
+	// The total "Stride In Air" is clamped internally at 95%
+	float AirExtensionAtMaxSpeed = 0.2f;
+	// Default is 6. Range is 0-inf. The height to raise foot (in cm) at the peak of the largest stride
+	float StepHeight = 10.f;
+	// Default is 0.5. Range is 0-1. How quickly the foot accelerates when lifting off the ground. At 0, it instantly moves, at 1 it eases in.
+	float StepEaseIn = 0.5f;
+	// Default is 0.2. Range is 0-1. How quickly the foot decelerates when planting. At 0, it instantly stops on contact, at 1 it eases into contact.
+	float StepEaseOut = 0.2f;
+	// Default is true. When true, feet are treated as non-overlapping circles and will not overlap each other or cross over.
+	bool bEnableFootCollision = false;
+	// Default is 1.0. Range is 0-inf. Global scale multiplier for all foot collision radii.
+	float FootCollisionGlobalScale = 1.0f;
+	// Default is true. When true, foot height is adjusted using sphere-casting against the world.
+	bool bEnableGroundCollision = true;
+	// Default is 30. Range is 0-inf. The maximum height in cm to place foot on-top of collision geo (may move foot up or down by this amount).
+	float MaxCollisionHeight = 30.0f;
+	// The 'channel' used to trace spheres into the physics world for collision detection. Determines which components to hit.
+	TEnumAsByte<ETraceTypeQuery> TraceChannel;
+	// Default is 0.8. Range is 0-1. When Ground Collision is enabled, this will orient the pitch of the foot to match the slope of the ground.
+	float OrientFootToGroundPitch = 0.8f;
+	// Default is 0.8. Range is 0-1. When Ground Collision is enabled, this will orient the roll of the foot to match the slope of the ground.
+	float OrientFootToGroundRoll = 0.5f;
+	// actor and component to ignore in collision queries
+	const AActor* IgnoredActor = nullptr;
+	const UActorComponent* IgnoredComponent = nullptr;
+	// the world to collide against
+	const UWorld* World = nullptr;
+};
+
+struct FLocomotorPelvisSettings
+{
+	// this represents the current input pose of the skeleton coming from animation
+	// this allows the locomotor to adapt dynamically to changing animation
+	FTransform InputPelvisComponentSpace;
+	// Default is 0.05. Typical range is 0.05-1. The time (in seconds) it takes for the Pelvis to reach half way to it's position goal.
+	float PositionDampingHalfLife = 0.05f;
+	// Default is 2.0. Typical range is 0-20. Higher values cause pelvis to rotate more rapidly towards goal orientation.
+	float RotationStiffness = 40.f;
+	// Default is 2.0. Typical range is 0-20. Higher values cause pelvis rotation to dampen quickly.
+	float RotationDamping = .9f;
+	// Default is 3.0. Typical range is 1-3. Higher values cause pelvis to move further ahead of the feet.
+	float LeadAmount = 3.f;
+	// Default is 0.1. Typical range is 0.1-1. The time (in seconds) it takes for the Pelvis to reach half way to it's lead goal.
+	float LeadDampingHalfLife = 0.1f;
+	// Default is -3.f. Range is -inf to -inf. The amount to lower the pelvis at the minimum height in the cycle. 
+	float BobOffset = -4.f;
+	// Default is 150.0. Typical range is 1-200. Higher values cause pelvis to bob more rapidly towards target height.
+	float BobStiffness = 150.f;
+	// Default is 0.9. Typical range is 0-2. Higher values cause pelvis bob to dampen quickly.
+	float BobDamping = 1.25f;
+	// Default is -0.3. Typical range is -1 to 1. Pitch the pelvis up/down with ground. 0 is no effect, 1/-1 matches ground orientation. Use negative values for bipeds, positive for quadrupeds.
+	float OrientToGroundPitch = -0.3f;
+	// Default is -0.3. Typical range is -1 to 1. Roll the pelvis side-to-side with ground. 0 is no effect, 1/-1 matches ground orientation. Use negative values for bipeds, positive for quadrupeds.
+	float OrientToGroundRoll = -0.3f;
+};
+
+struct FLocomotorSpineSettings
+{
+	// Default is 30. The maximum angle (in degrees) to lean the spine in the direction of travel.
+	float LeanAngleMax = 30.0f;
+	// Default is 150.0. Typical range is 1-200. Higher values cause spine to lean more rapidly towards target direction.
+	float LeanStiffness = 150.f;
+	// Default is 0.9. Typical range is 0-2. Higher values cause spine leaning to dampen quickly.
+	float LeanDamping = 0.9f;
+};
+
+struct FLocomotorHeadSettings
+{
+	// Default is 30. The maximum angle (in degrees) to rotate the head to look at the goal.
+	float AngleMax = 30.0f;
+	// Default is 150.0. Typical range is 1-200. Higher values cause spine to lean more rapidly towards target direction.
+	float Stiffness = 150.f;
+	// Default is 0.9. Typical range is 0-2. Higher values cause spine leaning to dampen quickly.
+	float Damping = 0.9f;
+};
+
+struct FLocomotorFootSettings
+{
+	// the radius in cm to avoid feet penetrations
+	float CollisionRadius = 10.f;
+	// Default is (0, 0, 50). The max rotation to peel the heel before lift off. This is in the local space of the ankle bone.
+	FVector MaxHeelPeel = FVector(0.f, 0.f, 50.f);
+	// either 0.0f or 0.5f depending on order in set (alternates Left/Right/Left/Right with values 0 / 0.5 / 0 / 0.5 etc...)
+	float StaticPhaseOffset = 0.f;
+	// applies a static offset that shifts feet in the local space of the foot bone in the reference pose
+	FVector StaticLocalOffset = FVector::ZeroVector;
+	// the index of the set this foot belongs to
+	int32 FootSetIndex = 0;
+};
+
+struct FLocomotorInputSettings
+{
+	// the current target goal location in world space
+	FTransform CurrentWorldRootGoal;
+	// time in seconds since last update
+	float DeltaTime = 0.16f;
+	// all settings related to overall movement and strides
+	FLocomotorMovementSettings Movement;
+	// all settings related to steps
+	FLocomotorStepSettings Stepping;
+	// pelvis settings
+	FLocomotorPelvisSettings Pelvis;
+	// spine settings
+	FLocomotorSpineSettings Spine;
+};
+
+struct FPlaneConstraint
+{
+	FPlaneConstraint() : Active(false), Normal(FVector::UpVector) {}
+	
+	bool Active;
+	FVector Normal;
+};
+
+struct FLocomotorFoot
+{
+	FLocomotorFoot(
+		const FTransform& InitialWorldRootGoal,
+		const FTransform& InitialWorldFoot,
+		const FLocomotorFootSettings& InSettings);
+
+	FLocomotorFootSettings Settings;
+
+	// the current phase gradually approaches this value, we don't directly set the current phase to avoid foot popping  
+	float TargetPhase = 0.f;
+	// the current phase of this foot to use for animation
+	float CurrentPhase = 0.f;
+	// the point in the foot's phase when the swing begins (variable with speed/frequency)
+	float FootPhaseWhenSwingEnds = 0.5f;
+	// whether this foot is close enough to the final goal to stop animating
+	bool bAtRest = true;
+	// current state of the foot (either swinging or planted)
+	bool bInSwingPhase = false;
+	// only true on the tick when this foot transitions to bInSwingPhase=true
+	bool bUnplantingThisTick = false;
+	// does this foot currently have a target beyond the min step threshold?
+	bool bWantsToStep = false;
+	// the distance from the planted location to the final target
+	FVector PrevToFinalTargetNorm;
+	float PrevToFinalTargetDistance;
+
+	// the initial transform relative to the initial root goal
+	FTransform InitialRelativeToRootGoal;
+	// the initial transform relative to the initial body
+	FTransform InitialRelativeToBody;
+	// initial transform
+	FTransform InitialWorld;
+
+	// the max height of the foot for the current stride (calculated once when unplanting)
+	float CurrentStrideHeight;
+	// the neutral location at the start of a step, may be different from the actual planted location
+	FTransform StrideOriginWorld;
+	// the prev foot plant in world space
+	FTransform PlantedWorld;
+	// the current target location to step towards
+	FTransform CurrentTargetWorld;
+	// a heavily damped spring for continuously updating target
+	FVectorSpring CurrentTargetSpring;
+	// the final destination location for foot
+	FTransform FinalTargetWorld;
+	
+	// the current world position of this foot with no height added
+	FVector CurrentWorldFlatPosition;
+	// the current world position of this foot with no height added and no ease in/out applied
+	FVector CurrentWorldFlatPositionNoEase;
+	// the current world transform of this foot
+	FTransform CurrentWorld;
+	// the world transform of this foot at the end of the previous tick
+	FTransform PrevWorld;
+
+	// store current height offset for bobbing
+	float CurrentHeightOffset;
+
+	// rotation is smoothly blended through a spring
+	FQuatSpring CurrentRotationSpring;
+};
+
+struct FLocomotorFootSet
+{
+	FLocomotorFootSet(const float InPhaseOffset, int32 InSetIndex) : PhaseOffset(InPhaseOffset), SetIndex(InSetIndex){};
+
+	// add a foot, returns the index of the foot or -1 if no set found
+	int32 AddFoot(
+		const FTransform& InitialWorldRootGoal,
+		const FTransform& InitialWorldFoot,
+		const FLocomotorFootSettings& InSettings);
+
+	// get access to the feet in this set
+	TArray<FLocomotorFoot>& GetFeet() {return Feet; };
+	const TArray<FLocomotorFoot>& GetFeet() const {return Feet; };
+
+	// get number of feet in this set
+	int32 GetNumFeet() const { return Feet.Num(); };
+
+	// get the current offset of this set
+	float GetPhaseOffset() const { return PhaseOffset; };
+	
+private:
+	
+	// the 0-1 phase offset for this set of feet (relative to the global phase)
+	float PhaseOffset = 0.f;
+	// the array of feet belonging to this set
+	TArray<FLocomotorFoot> Feet;
+	// the index of this set
+	int32 SetIndex = 0;
+};
+
+struct FLocomotorHead
+{
+	FLocomotorHead() = default;
+
+	void Initialize(const FTransform& InInitialWorld);
+	
+	// the initial world transform of the head bone
+	FTransform InitialWorld;
+	FTransform CurrentWorld;
+	FQuatSpring RotationSpring;
+
+	bool bInitialized = false;
+};
+
+struct FLocomotorPelvis
+{
+	FLocomotorPelvis() = default;
+
+	void Initialize(const TArray<FLocomotorFoot*>& AllFeet, const FTransform& InitialRootGoalWorld);
+
+	void CalculateBodyTransform(
+		const TArray<FLocomotorFoot*> AllFeet, 
+		const FTransform& RootGoal,
+		FTransform& OutBodyTransform);
+	
+	// the initial world transform of the pelvis
+	FTransform InitialWorld;
+	FTransform InitialRelativeToRootGoal;
+	FTransform CurrentWorld;
+	
+	// the transform of the "body"; this is a separately maintained transform representing the center of mass
+	// which is different than the pelvis location which is usually far from the center of the body in quadrupeds
+	FTransform InitialBodyWorld;
+	FTransform CurrentBodyWorld;
+	FTransform PreviousBodyTargetWorld;
+	FTransform CurrentBodyTargetWorld;
+	FVector CurrentBodyLead;
+	FVectorDamper BodyLeadDamper;
+	FQuatSpring BodyRotationSpring;
+	FVectorDamper PelvisPositionDamper;
+
+	// height offset (bob) is tracked separately and applied additively through a spring
+	float CurrentHeightOffset;
+	FFloatSpring BobOffsetSpring;
+};
+
+USTRUCT()
+struct FLocomotor
+{
+	GENERATED_BODY()
+	
+public:
+	
+	FLocomotor() = default;
+
+	// initialize all internal data structures (can be called multiple times)
+	void Reset(
+		const FTransform& InInitialRootGoalWorld,
+		const FTransform& InInitialWorldPelvis);
+
+	// add a set of feet, returns the foot set
+	int32 AddFootSet(const float PhaseOffset);
+
+	// add a foot, returns the index of the foot or -1 if no set found
+	int32 AddFootToSet(
+		const int32 FootSetIndex,
+		const FTransform& InitialWorldFoot,
+		const FLocomotorFootSettings& InSettings);
+
+	// set the head transform and settings
+	void SetHead(const FTransform& InitialWorldTransform);
+
+	// update the simulation state and generate new foot transforms
+	void RunSimulation(const FLocomotorInputSettings& Settings);
+
+	// must have at least one set of feet to progress
+	bool HasFeet() const { return !AllFeet.IsEmpty(); };
+
+	// read-only access to feet
+	const TArray<FLocomotorFootSet>& GetFootSets() const { return FootSets; };
+
+	// read-only access to pelvis transform
+	const FTransform& GetPelvisCurrent() const { return Pelvis.CurrentWorld; };
+
+	// read-only access to body transform
+	const FTransform& GetBodyCurrent() const { return Pelvis.CurrentBodyWorld; };
+
+	// read-only access to pelvis transform
+	const FTransform& GetBodyTarget() const { return Pelvis.CurrentBodyTargetWorld; };
+
+	// get the current phase
+	float GetPhaseCurrent() const;
+
+	// read-only access to settings
+	const FLocomotorInputSettings& GetSettings() const { return Settings; };
+
+	// get final foot transforms
+	void GetFeetCurrent(TArray<FTransform> OutTransforms) const;
+
+	// get access to all feet
+	TArray<FLocomotorFoot*>& GetFeet() {return AllFeet; };
+	const TArray<FLocomotorFoot*>& GetFeet() const {return AllFeet; };
+
+private:
+
+	void Simulate();
+	bool Initialize();
+	void UpdateWorldSpeedAndPhase();
+	void UpdateFeetTargets();
+	void AnimateFeet();
+	void UpdateBody();
+	void AnimatePelvis();
+	void AnimateHead();
+	
+	void ResolveFootToFootCollision(FLocomotorFoot& Foot);
+	void PushCirclesApartInFloorPlane(
+		FVector& CenterA,
+		FVector& CenterB,
+		float RadiusA,
+		float RadiusB,
+		float InvMassA,
+		float InvMassB);
+	FVector ProjectToGroundWithSphereCast(const FLocomotorFoot& Foot, FTransform& TransformToProject);
+
+	float EaseInOutNorm(const float Input, float EaseInAmount, float EaseOutAmount) const;
+	static float WrapPhaseInRange(const float PhaseToWrap);
+	static float BlendTowardsTargetPhase(const float CurrentPhase, const float TargetPhase, const float PhaseBlendRate);
+	float GetPhaseOffsetForSetFromMovementStyle(int32 SetIndex);
+
+	bool CastSphere(const FVector& Start, const FVector& End, float SphereRadius, FHitResult& OutHitResult);
+
+	void CalcPitchRollFromNormal(
+		const FVector& Fwd,
+		const FVector& Up,
+		const FVector& Normal,
+		FVector& OutSide,
+		FQuat& OutPitchRotation,
+		FQuat& OutRollRotation);
+	
+	// the current settings (updated each frame)
+	FLocomotorInputSettings Settings;
+	// the root goal in world space at the start of the simulation
+	FTransform InitialRootGoalWorld;
+	// current state of each foot being animated
+	TArray<FLocomotorFootSet> FootSets;
+	// pointers to all feet in all sets (for convenience)
+	TArray<FLocomotorFoot*> AllFeet;
+	// pelvis data
+	FLocomotorPelvis Pelvis;
+	// head data
+	FLocomotorHead Head;
+	
+	// global phase and movement speed (dynamically adjusted during simulation)
+	float CurrentGlobalPhase;
+	float CurrentPhaseSpeed;
+	float CurrentSpeed;
+	float CurrentPercentOfMaxSpeed;
+	float CurrentStrideLength;
+	bool bFullyAtRest = true;
+	bool bAccelerating = true;
+	float AccumulatedTimeToSimulate;
+
+	// whether the locomotor has cached it's internal state on first update
+	bool bPostInitialized = false;
+};
diff --git a/Engine/Plugins/Experimental/Animation/Locomotor/Source/Locomotor/Public/RigUnit_Locomotor.h b/Engine/Plugins/Experimental/Animation/Locomotor/Source/Locomotor/Public/RigUnit_Locomotor.h
new file mode 100644
--- /dev/null
+++ b/Engine/Plugins/Experimental/Animation/Locomotor/Source/Locomotor/Public/RigUnit_Locomotor.h
@@ -0,0 +1,396 @@
+﻿// Copyright Epic Games, Inc. All Rights Reserved.
+
+#pragma once
+
+#include "LocomotorCore.h"
+#include "Units/RigUnitContext.h"
+#include "Units/Highlevel/RigUnit_HighlevelBase.h"
+
+#include "RigUnit_Locomotor.generated.h"
+
+USTRUCT()
+struct LOCOMOTOR_API FLocomotorDebugSettings
+{
+	GENERATED_BODY()
+
+	FLocomotorDebugSettings() :
+		bDrawDebug(true),
+		bDrawBody(true),
+		bDrawPhaseCircle(true),
+		bDrawCurrentFeet(true),
+		bDrawCurrentFeetFlat(true),
+		bDrawCurrentFeetTarget(true),
+		bDrawPrevFeetTarget(true),
+		bDrawFinalFeetTarget(true),
+		bDrawFeetCollision(true),
+		Scale(25.0f),
+		Thickness(0.5f) {}
+
+	UPROPERTY(meta = (Input))
+	bool bDrawDebug;
+
+	UPROPERTY(meta = (Input))
+	bool bDrawBody;
+
+	UPROPERTY(meta = (Input))
+	bool bDrawPhaseCircle;
+	
+	UPROPERTY(meta = (Input))
+	bool bDrawCurrentFeet;
+
+	UPROPERTY(meta = (Input))
+	bool bDrawCurrentFeetFlat;
+
+	UPROPERTY(meta = (Input))
+	bool bDrawCurrentFeetTarget;
+	
+	UPROPERTY(meta = (Input))
+	bool bDrawPrevFeetTarget;
+
+	UPROPERTY(meta = (Input))
+	bool bDrawFinalFeetTarget;
+
+	UPROPERTY(meta = (Input))
+	bool bDrawFeetCollision;
+	
+	UPROPERTY(meta = (Input, ClampMin = "0.0", UIMin = "0.01", UIMax = "100.0"))
+	float Scale;
+
+	UPROPERTY(meta = (Input, ClampMin = "0.0", UIMin = "0.01", UIMax = "100.0"))
+	float Thickness;
+	
+	void DrawDebug(FRigVMDrawInterface* DrawInterface, const FLocomotor& Locomotor, const FTransform& WorldToGlobal) const;
+};
+
+UENUM()
+enum class ELocomotorMovementStyle : uint8
+{
+	Walk,
+	Trot,
+	Gallop,
+};
+
+USTRUCT()
+struct LOCOMOTOR_API FMovementSettings
+{
+	GENERATED_BODY()
+
+	FMovementSettings() : 
+		MinimumStepLength(10.f),
+		SpeedMax(80.f),
+		SpeedMin(50.f),
+		Acceleration(100.f),
+		Deceleration(30.f),
+		GlobalTimeScale(1.f){}
+	
+	// the length of the smallest step to take in cm
+	UPROPERTY(meta = (Input))
+	float MinimumStepLength = 10.f;
+	
+	// the maximum movement speed in cm/s
+	UPROPERTY(meta = (Input))
+	float SpeedMax = 80.f;
+
+	// the minimum movement speed in cm/s
+	UPROPERTY(meta = (Input))
+	float SpeedMin = 50.f;
+
+	// the maximum phase speed in cycles/s
+	UPROPERTY(meta = (Input))
+	float PhaseSpeedMax = 4.f;
+	
+	// the minimum phase speed in cycles/s
+	UPROPERTY(meta = (Input))
+	float PhaseSpeedMin = 1.0f;
+
+	// the acceleration in cm/s/s (ie, how quickly the movement speed increases to reach the goal)
+	UPROPERTY(meta = (Input))
+	float Acceleration = 100.f;
+
+	// the acceleration in cm/s/s (ie, how quickly the movement speed decreases when approaching the goal)
+	UPROPERTY(meta = (Input))
+	float Deceleration = 30.f;
+
+	// adjust global speed of the simulation
+	UPROPERTY(meta = (Input))
+	float GlobalTimeScale = 1.f;
+
+	// The the movement styles to cycle through based on speed
+	UPROPERTY(meta = (Input))
+	TArray<ELocomotorMovementStyle> Styles;
+};
+
+USTRUCT()
+struct LOCOMOTOR_API FStepSettings
+{
+	GENERATED_BODY()
+
+	FStepSettings() :
+		PercentOfStrideInAir(0.35f),
+		AirExtensionAtMaxSpeed(0.2f),
+		StepHeight(6.f),
+		StepEaseIn(0.5f),
+		StepEaseOut(0.2f),
+		bEnableFootCollision(true),
+		FootCollisionGlobalScale(1.0f),
+		bEnableGroundCollision(true){}
+	
+	// Default is 0.35. Range is 0.1 - 0.9. The percentage amount of the 0-1 stride phase where the foot is in the air.
+	// At the default (0.35), the foot will plant on the ground 35% through it's total phase.
+	// At lower values, the foot will spend more time on the ground as a percentage of the total stride time.
+	// At higher values, each foot will spend less time on the ground as a percentage of the total stride time.
+	UPROPERTY(meta = (Input, UIMin = "0.1", UIMax=".9", ClampMin="0.1", ClampMax=".9"))
+	float PercentOfStrideInAir = 0.35f;
+	// Default is 0.2. Range is 0.0 - 0.9. This value is added to the "PercentOfStrideInAir" to extend the air-time while at max speed.
+	// For example, if PercentOfStrideInAir is at 0.5, and the air extension is at 0.4 the foot will be in the air for 90% of the stride cycle.
+	// This value is gradually blended on as the Speed accelerates between SpeedMin and SpeedMax.
+	// The total "Stride In Air" is clamped internally at 95%
+	UPROPERTY(meta = (Input, UIMin = "0.1", UIMax=".9", ClampMin="0.1", ClampMax=".9"))
+	float AirExtensionAtMaxSpeed = 0.2f;
+	
+	// Default is 6. Range is 0-inf. The height to raise foot (in cm) at the peak of the largest stride
+	UPROPERTY(meta = (Input, UIMin = "0.1", UIMax="50"))
+	float StepHeight = 6.f;
+	// Default is 0.5. Range is 0-1. How quickly the foot accelerates when lifting off the ground. At 0, it instantly moves, at 1 it eases in.
+	UPROPERTY(meta = (Input, UIMin = "0", UIMax="1", ClampMin="0", ClampMax="1"))
+	float StepEaseIn = 0.5f;
+	// Default is 0.2. Range is 0-1. How quickly the foot decelerates when planting. At 0, it instantly stops on contact, at 1 it eases into contact.
+	UPROPERTY(meta = (Input, UIMin = "0", UIMax="1", ClampMin="0", ClampMax="1"))
+	float StepEaseOut = 0.2f;
+	
+	// Default is true. When true, feet are treated as non-overlapping circles and will not overlap each other or cross over.
+	UPROPERTY(meta = (Input))
+	bool bEnableFootCollision = true;
+	// Default is 1.0. Range is 0-inf. Global scale multiplier for all foot collision radii.
+	UPROPERTY(meta = (Input))
+	float FootCollisionGlobalScale = 1.0f;
+	// Default is true. When true, foot height is adjusted using sphere-casting against the world.
+	UPROPERTY(meta = (Input))
+	bool bEnableGroundCollision = true;
+	// Default is 30. Range is 0-inf. The maximum height in cm to place foot on-top of collision geo (may move foot up or down by this amount).
+	UPROPERTY(meta = (Input, UIMin = "0"))
+	float MaxCollisionHeight = 30.0f;
+	// The 'channel' used to trace spheres into the physics world for collision detection. Determines which components to hit.
+	UPROPERTY(meta = (Input))
+	TEnumAsByte<ETraceTypeQuery> TraceChannel;
+	// Default is 0.8. Range is 0-1. When Ground Collision is enabled, this will orient the pitch of the foot to match the slope of the ground.
+	UPROPERTY(meta = (Input, UIMin = "0", UIMax="1", ClampMin="0", ClampMax="1"))
+	float OrientFootToGroundPitch = 0.8f;
+	// Default is 0.8. Range is 0-1. When Ground Collision is enabled, this will orient the roll of the foot to match the slope of the ground.
+	UPROPERTY(meta = (Input, UIMin = "0", UIMax="1", ClampMin="0", ClampMax="1"))
+	float OrientFootToGroundRoll = 0.5f;
+};
+
+USTRUCT()
+struct LOCOMOTOR_API FPelvisSettings
+{
+	GENERATED_BODY()
+
+	FPelvisSettings() :
+		PelvisBone(FRigElementKey(NAME_None, ERigElementType::Bone)),
+		PositionDampingHalfLife(0.1f),
+		RotationStiffness(40.0f),
+		RotationDamping(0.9f),
+		LeadAmount(2.0f),
+		LeadDampingHalfLife(0.1f),
+		BobOffset(-8.f),
+		BobStiffness(40.f),
+		BobDamping(.9f),
+		OrientToGroundPitch(-0.3f),
+		OrientToGroundRoll(-0.3f){};
+	
+	// Usually the Pelvis or Hip bone. The bone that translates the whole body (not necessarily the root of the entire hierarchy).
+	UPROPERTY(meta = (Input))
+	FRigElementKey PelvisBone = FRigElementKey(NAME_None, ERigElementType::Bone);
+	
+	// Default is 0.1. Typical range is 0.1-1. The time (in seconds) it takes for the Pelvis to reach half way to it's goal.
+	UPROPERTY(meta = (Input, UIMin = "0.0", UIMax = "10.0"))
+	float PositionDampingHalfLife = 0.1f;
+
+	// Default is 40.0. Typical range is 1-100. Higher values cause pelvis to rotate more rapidly towards goal orientation.
+	UPROPERTY(meta = (Input, UIMin = "0.1", UIMax = "100.0"))
+	float RotationStiffness = 40.f;
+	// Default is 2.0. Typical range is 0-2. Higher values cause pelvis rotation to dampen quickly.
+	UPROPERTY(meta = (Input, UIMin = "0.1", UIMax = "2.0"))
+	float RotationDamping = .9f;
+
+	// Default is 2.0. Typical range is 1-3. Higher values cause pelvis to move further ahead of the feet.
+	UPROPERTY(meta = (Input, UIMin = "0.1", UIMax = "20.0"))
+	float LeadAmount = 2.f;
+	// Default is 0.1. Typical range is 0.1-1. The time (in seconds) it takes for the Pelvis to reach half way to it's lead goal.
+	UPROPERTY(meta = (Input, UIMin = "0.0", UIMax = "10.0"))
+	float LeadDampingHalfLife = 0.1f;
+	
+	// Default is 10.f. Range is 0-inf. The amount to lower the pelvis at the minimum height in the cycle.
+	UPROPERTY(meta = (Input, UIMin = "0"))
+	float BobOffset = -8.f;
+	// Default is 40.0. Typical range is 1-100. Higher values cause pelvis to bob more rapidly towards target height.
+	UPROPERTY(meta = (Input, UIMin = "0.1", UIMax = "20.0"))
+	float BobStiffness = 40.f;
+	// Default is 0.9. Typical range is 0-2. Higher values cause pelvis bob to dampen quickly.
+	UPROPERTY(meta = (Input, UIMin = "0.1", UIMax = "20.0"))
+	float BobDamping = .9f;
+
+	// Default is -0.3. Typical range is -1 to 1. Pitch the pelvis up/down with ground. 0 is no effect, 1/-1 matches ground orientation. Use negative values for bipeds, positive for quadrupeds.
+	UPROPERTY(meta = (Input, UIMin = "-1", UIMax = "1.0"))
+	float OrientToGroundPitch = -0.3f;
+	// Default is -0.3. Typical range is -1 to 1. Roll the pelvis side-to-side with ground. 0 is no effect, 1/-1 matches ground orientation. Use negative values for bipeds, positive for quadrupeds.
+	UPROPERTY(meta = (Input, UIMin = "-1", UIMax = "1.0"))
+	float OrientToGroundRoll = -0.3f;
+};
+
+USTRUCT()
+struct LOCOMOTOR_API FSpineSettings
+{
+	GENERATED_BODY()
+
+	FSpineSettings() :
+		SpineBone(FRigElementKey(NAME_None, ERigElementType::Bone)),
+		PitchAngleMax(30.0f),
+		PitchStiffness(150.f),
+		PitchDamping(0.9f){}
+	
+	// The base spine bone. Usually directly below the Pelvis bone. The bone that rotates to lean the whole spine.
+	UPROPERTY(meta = (Input))
+	FRigElementKey SpineBone = FRigElementKey(NAME_None, ERigElementType::Bone);
+
+	// Default is 30. The maximum angle (in degrees) to lean the spine in the direction of travel.
+	UPROPERTY(meta = (Input, UIMin = "0.1", UIMax = "180.0"))
+	float PitchAngleMax = 30.0f;
+	
+	// Default is 150.0. Typical range is 1-200. Higher values cause spine to lean more rapidly towards target direction.
+	UPROPERTY(meta = (Input, UIMin = "0.1", UIMax = "200.0"))
+	float PitchStiffness = 150.f;
+	
+	// Default is 0.9. Typical range is 0-2. Higher values cause spine leaning to dampen quickly.
+	UPROPERTY(meta = (Input, UIMin = "0.1", UIMax = "20.0"))
+	float PitchDamping = 0.9f;
+};
+
+USTRUCT()
+struct LOCOMOTOR_API FHeadSettings
+{
+	GENERATED_BODY()
+
+	FHeadSettings() :
+		HeadBone(FRigElementKey(NAME_None, ERigElementType::Bone)),
+		AngleMax(30.0f),
+		Stiffness(150.f),
+		Damping(0.9f){}
+	
+	// The base spine bone. Usually directly below the Pelvis bone. The bone that rotates to lean the whole spine.
+	UPROPERTY(meta = (Input))
+	FRigElementKey HeadBone = FRigElementKey(NAME_None, ERigElementType::Bone);
+
+	// Default is 30. The maximum angle (in degrees) to rotate the head to look at the goal.
+	UPROPERTY(meta = (Input, UIMin = "0.1", UIMax = "180.0"))
+	float AngleMax = 30.0f;
+	
+	// Default is 150.0. Typical range is 1-200. Higher values cause spine to lean more rapidly towards target direction.
+	UPROPERTY(meta = (Input, UIMin = "0.1", UIMax = "200.0"))
+	float Stiffness = 150.f;
+	
+	// Default is 0.9. Typical range is 0-2. Higher values cause spine leaning to dampen quickly.
+	UPROPERTY(meta = (Input, UIMin = "0.1", UIMax = "20.0"))
+	float Damping = 0.9f;
+};
+
+
+USTRUCT()
+struct LOCOMOTOR_API FFootSettings
+{
+	GENERATED_BODY()
+
+	FFootSettings() = default;
+
+	// The bone at the END of the limb (typically the BALL of the foot). This is where the output goal originates.
+	UPROPERTY(meta = (Input))
+	FRigElementKey AnkleBone;
+
+	UPROPERTY(meta = (Input))
+	float CollisionRadius = 10.f;
+
+	// Default is (0, 0, 50). The max rotation to peel the heel before lift off. This is in the local space of the ankle bone.
+	UPROPERTY(meta = (Input))
+	FVector MaxHeelPeel = FVector(0.f, 0.f, 50.f);
+
+	// applies a static offset that shifts feet in the local space of the foot bone in the reference pose
+	UPROPERTY(meta = (Input))
+	FVector StaticLocalOffset = FVector::ZeroVector;
+};
+
+USTRUCT()
+struct LOCOMOTOR_API FFootSet
+{
+	GENERATED_BODY()
+
+	FFootSet() = default;
+
+	// All the feet belonging to this set. Feet in this set alternate phases (ie Left/Right/Left/Right).
+	UPROPERTY(meta = (Input))
+	TArray<FFootSettings> Feet;
+
+	// the 0-1 phase offset for this set of feet (relative to the global phase)
+	UPROPERTY(meta = (Input, UIMin = "0.0", UIMax = "1.0"))
+	float PhaseOffset = 0.f;
+};
+
+USTRUCT(meta=(DisplayName="Locomotor", Category="Simulation", Keywords="Walk Step Walker Walking Locomotion Procedural"))
+struct LOCOMOTOR_API FRigUnit_Locomotor: public FRigUnit_HighlevelBaseMutable
+{
+	GENERATED_BODY()
+
+	RIGVM_METHOD()
+	virtual void Execute() override;
+
+	FRigUnit_Locomotor() :
+		RootControl(NAME_None),
+		CachedRootElement(FCachedRigElement()){}
+
+	// list of limbs to walk on.
+	UPROPERTY(meta = (Output, DefaultArraySize = 2))
+	TArray<FTransform> FeetTransforms;
+
+	// the target location and rotation, in world space
+	// the locomotor will generate steps to pull the character towards this location
+	// footsteps will be oriented match the root goal's final orientation
+	UPROPERTY(meta = (Input, CustomWidget = "ControlName" ))
+	FName RootControl;
+	// Used to cache the root element
+	UPROPERTY()
+	FCachedRigElement CachedRootElement;
+
+	// the movement settings
+	UPROPERTY(meta = (Input))
+	FMovementSettings Movement;
+
+	// the set of feet belonging to the creature
+	UPROPERTY(meta = (Input, Constant))
+	TArray<FFootSet> FootSets;
+
+	// the global step settings
+	UPROPERTY(meta = (Input))
+	FStepSettings Stepping;
+
+	// the Pelvis or Hip bone to move with the feet
+	UPROPERTY(meta = (Input))
+	FPelvisSettings Pelvis;
+
+	// TODO
+	// the spine settings
+	//UPROPERTY(meta = (Input))
+	//FSpineSettings Spine;
+
+	// TODO
+	// the head settings
+	//UPROPERTY(meta = (Input))
+	//FHeadSettings Head;
+
+	// adjust debug drawing.
+	UPROPERTY(meta = (Input))
+	FLocomotorDebugSettings Debug;
+
+	// the locomotion simulator
+	UPROPERTY(Transient)
+	FLocomotor Locomotor;
+};
