Locomotor: Support for teleportation

#rnx

diff --git a/Engine/Plugins/Experimental/Animation/Locomotor/Source/Locomotor/Private/LocomotorCore.cpp b/Engine/Plugins/Experimental/Animation/Locomotor/Source/Locomotor/Private/LocomotorCore.cpp
--- a/Engine/Plugins/Experimental/Animation/Locomotor/Source/Locomotor/Private/LocomotorCore.cpp
+++ b/Engine/Plugins/Experimental/Animation/Locomotor/Source/Locomotor/Private/LocomotorCore.cpp
@@ -266,6 +266,15 @@ void FLocomotor::UpdateFeetTargets()
 		
 		// continuously update the final foot target
 		Foot.FinalTargetWorld = Foot.InitialRelativeToRootGoal * Settings.CurrentWorldRootGoal;
+		
+		if (Settings.Movement.bTeleport)
+		{
+			Foot.CurrentWorld = Foot.PlantedWorld = Foot.CurrentTargetWorld = Foot.FinalTargetWorld;
+			Foot.bAtRest = true;
+			Foot.bInSwingPhase = Foot.bUnplantingThisTick = Foot.bWantsToStep = false;
+			continue;
+		}
+		
 		if (Settings.Stepping.bEnableGroundCollision)
 		{
 			ProjectToGroundWithSphereCast(Foot, Foot.FinalTargetWorld);	
@@ -277,6 +286,11 @@ void FLocomotor::UpdateFeetTargets()
 		Foot.PrevToFinalTargetNorm.Normalize();
 	}
 
+	if (Settings.Movement.bTeleport)
+	{
+		return;
+	}
+
 	// continuously update "bWantsToStep"
 	// (based on if the foot's final target is beyond the min stride distance)
 	bool bAnyFeetWantToStep = false;
@@ -710,6 +724,17 @@ void FLocomotor::UpdateBody()
 		Pelvis.CalculateBodyTransform(AllFeet, Settings.CurrentWorldRootGoal, Pelvis.CurrentBodyTargetWorld);	
 	}
 
+	if (Settings.Movement.bTeleport)
+	{
+		Pelvis.CurrentBodyWorld = Pelvis.PreviousBodyTargetWorld = Pelvis.CurrentBodyTargetWorld;
+		Pelvis.CurrentBodyLead = FVector::ZeroVector;
+		Pelvis.BobOffsetSpring.Reset(0);
+		Pelvis.BodyLeadDamper.Reset(FVector::ZeroVector);
+		Pelvis.BodyRotationSpring.Reset(Pelvis.CurrentBodyWorld.GetRotation());
+		Pelvis.CurrentHeightOffset = Pelvis.CurrentBodyTargetWorld.GetLocation().Z;
+		return;
+	}
+
 	// pitch/roll the body target based on the slope of the ground
 	const bool bPitchBodyWithGround = !FMath::IsNearlyZero(Settings.Pelvis.OrientToGroundPitch);
 	const bool bRollBodyWithGround = !FMath::IsNearlyZero(Settings.Pelvis.OrientToGroundRoll);
diff --git a/Engine/Plugins/Experimental/Animation/Locomotor/Source/Locomotor/Private/RigUnit_Locomotor.cpp b/Engine/Plugins/Experimental/Animation/Locomotor/Source/Locomotor/Private/RigUnit_Locomotor.cpp
--- a/Engine/Plugins/Experimental/Animation/Locomotor/Source/Locomotor/Private/RigUnit_Locomotor.cpp
+++ b/Engine/Plugins/Experimental/Animation/Locomotor/Source/Locomotor/Private/RigUnit_Locomotor.cpp
@@ -43,7 +43,7 @@ FRigUnit_Locomotor_Execute()
 		const FTransform InitialRootGoalGlobal = Hierarchy->GetInitialGlobalTransform(CachedRootElement.GetIndex());
 		const FTransform InitialPelvisGlobal = Hierarchy->GetInitialGlobalTransform(Pelvis.PelvisBone);
 		const FTransform InitialPelvisWorld = InitialPelvisGlobal * ToWorld;
-		const FTransform InitialRootGoalWorld = InitialRootGoalGlobal.GetRelativeTransform(InitialPelvisGlobal) * InitialPelvisWorld;
+		const FTransform InitialRootGoalWorld = InitialRootGoalGlobal * ToWorld;
 		Locomotor.Reset(InitialRootGoalWorld, InitialPelvisWorld);
 
 		// add all the feet
@@ -87,7 +87,7 @@ FRigUnit_Locomotor_Execute()
 	
 	// update the simulation settings
 	FLocomotorInputSettings Settings;
-	Settings.CurrentWorldRootGoal = RootTransformNoScale;
+	Settings.CurrentWorldRootGoal = RootTransformNoScale * ToWorld;
 	Settings.DeltaTime = ExecuteContext.GetDeltaTime() * Movement.GlobalTimeScale;
 	// movement
 	Settings.Movement.MinimumStepLength = Movement.MinimumStepLength;
@@ -102,6 +102,7 @@ FRigUnit_Locomotor_Execute()
 		EMovementStyle StyleToAdd = static_cast<EMovementStyle>(Style); // cast from BP enum to internal enum
 		Settings.Movement.Styles.Add(StyleToAdd);
 	}
+	Settings.Movement.bTeleport = Movement.bTeleport;
 	Settings.Stepping.PercentOfStrideInAir = Stepping.PercentOfStrideInAir;
 	Settings.Stepping.AirExtensionAtMaxSpeed = Stepping.AirExtensionAtMaxSpeed;
 	Settings.Stepping.StepHeight = Stepping.StepHeight;
diff --git a/Engine/Plugins/Experimental/Animation/Locomotor/Source/Locomotor/Public/LocomotorCore.h b/Engine/Plugins/Experimental/Animation/Locomotor/Source/Locomotor/Public/LocomotorCore.h
--- a/Engine/Plugins/Experimental/Animation/Locomotor/Source/Locomotor/Public/LocomotorCore.h
+++ b/Engine/Plugins/Experimental/Animation/Locomotor/Source/Locomotor/Public/LocomotorCore.h
@@ -161,6 +161,8 @@ struct FLocomotorMovementSettings
 	float Deceleration  = 30.f;
 	// movement styles
 	TArray<EMovementStyle> Styles;
+	// teleportation
+	bool bTeleport = false;
 
 	//static TMap<EMovementStyle, float> StyleOffsets;
 };
diff --git a/Engine/Plugins/Experimental/Animation/Locomotor/Source/Locomotor/Public/RigUnit_Locomotor.h b/Engine/Plugins/Experimental/Animation/Locomotor/Source/Locomotor/Public/RigUnit_Locomotor.h
--- a/Engine/Plugins/Experimental/Animation/Locomotor/Source/Locomotor/Public/RigUnit_Locomotor.h
+++ b/Engine/Plugins/Experimental/Animation/Locomotor/Source/Locomotor/Public/RigUnit_Locomotor.h
@@ -81,7 +81,8 @@ struct LOCOMOTOR_API FMovementSettings
 		SpeedMin(50.f),
 		Acceleration(100.f),
 		Deceleration(30.f),
-		GlobalTimeScale(1.f){}
+		GlobalTimeScale(1.f),
+		bTeleport(false){}
 	
 	// the length of the smallest step to take in cm
 	UPROPERTY(meta = (Input))
@@ -115,6 +116,10 @@ struct LOCOMOTOR_API FMovementSettings
 	UPROPERTY(meta = (Input))
 	float GlobalTimeScale = 1.f;
 
+	// reset the simulation. this can be used to teleport the locomotor.
+	UPROPERTY(meta = (Input))
+	bool bTeleport;
+
 	// The the movement styles to cycle through based on speed
 	UPROPERTY(meta = (Input))
 	TArray<ELocomotorMovementStyle> Styles;
