[Locomotor] Fixed possible non-normalized rotation in foot simulation

#FailedValidations Preflight

diff --git a/Engine/Plugins/Experimental/Animation/Locomotor/Source/Locomotor/Private/LocomotorCore.cpp b/Engine/Plugins/Experimental/Animation/Locomotor/Source/Locomotor/Private/LocomotorCore.cpp
--- a/Engine/Plugins/Experimental/Animation/Locomotor/Source/Locomotor/Private/LocomotorCore.cpp
+++ b/Engine/Plugins/Experimental/Animation/Locomotor/Source/Locomotor/Private/LocomotorCore.cpp
@@ -18,6 +18,7 @@ FLocomotorFoot::FLocomotorFoot(
 	CurrentPhase = 1.0f;
 	FootPhaseWhenSwingEnds = 0.5f;
 	InitialRelativeToRootGoal = InitialWorldFoot.GetRelativeTransform(InitialWorldRootGoal);
+	InitialRelativeToRootGoal.NormalizeRotation();
 	InitialRelativeToBody = FTransform::Identity; // this is initialized after Pelvis calculates initial body transform
 	InitialWorld = InitialWorldFoot;
 	CurrentWorld = InitialWorld;
