[Locomotor] Fixed uninitialized member variables in FStepSettings

diff --git a/Engine/Plugins/Experimental/Animation/Locomotor/Source/Locomotor/Public/RigUnit_Locomotor.h b/Engine/Plugins/Experimental/Animation/Locomotor/Source/Locomotor/Public/RigUnit_Locomotor.h
--- a/Engine/Plugins/Experimental/Animation/Locomotor/Source/Locomotor/Public/RigUnit_Locomotor.h
+++ b/Engine/Plugins/Experimental/Animation/Locomotor/Source/Locomotor/Public/RigUnit_Locomotor.h
@@ -138,7 +138,11 @@ struct LOCOMOTOR_API FStepSettings
 		StepEaseOut(0.2f),
 		bEnableFootCollision(true),
 		FootCollisionGlobalScale(1.0f),
-		bEnableGroundCollision(true){}
+		bEnableGroundCollision(true),
+		MaxCollisionHeight(30.0f),
+		TraceChannel(ETraceTypeQuery::TraceTypeQuery1),
+		OrientFootToGroundPitch(0.8f),
+		OrientFootToGroundRoll(0.5f){}
 	
 	// Default is 0.35. Range is 0.1 - 0.9. The percentage amount of the 0-1 stride phase where the foot is in the air.
 	// At the default (0.35), the foot will plant on the ground 35% through it's total phase.
@@ -177,7 +181,7 @@ struct LOCOMOTOR_API FStepSettings
 	float MaxCollisionHeight = 30.0f;
 	// The 'channel' used to trace spheres into the physics world for collision detection. Determines which components to hit.
 	UPROPERTY(meta = (Input))
-	TEnumAsByte<ETraceTypeQuery> TraceChannel;
+	TEnumAsByte<ETraceTypeQuery> TraceChannel = ETraceTypeQuery::TraceTypeQuery1;
 	// Default is 0.8. Range is 0-1. When Ground Collision is enabled, this will orient the pitch of the foot to match the slope of the ground.
 	UPROPERTY(meta = (Input, UIMin = "0", UIMax="1", ClampMin="0", ClampMax="1"))
 	float OrientFootToGroundPitch = 0.8f;
