[Control Rig] Fix for volatile feet pointers in Locomotor node.
Added additional checks for phase input values being in range.

diff --git a/Engine/Plugins/Experimental/Animation/Locomotor/Source/Locomotor/Private/LocomotorCore.cpp b/Engine/Plugins/Experimental/Animation/Locomotor/Source/Locomotor/Private/LocomotorCore.cpp
--- a/Engine/Plugins/Experimental/Animation/Locomotor/Source/Locomotor/Private/LocomotorCore.cpp
+++ b/Engine/Plugins/Experimental/Animation/Locomotor/Source/Locomotor/Private/LocomotorCore.cpp
@@ -106,7 +106,8 @@ void FLocomotor::Reset(const FTransform& InInitialRootGoalWorld, const FTransfor
 int32 FLocomotor::AddFootSet(const float PhaseOffset)
 {
 	bPostInitialized = false;
-	return FootSets.Emplace(FLocomotorFootSet(PhaseOffset, FootSets.Num()));
+	float WrappedPhaseOffset = FMath::Fmod(FMath::Fmod(PhaseOffset, 1.0f) + 1.0f, 1.0f);
+	return FootSets.Emplace(FLocomotorFootSet(WrappedPhaseOffset, FootSets.Num()));
 }
 
 int32 FLocomotor::AddFootToSet(
@@ -122,7 +123,19 @@ int32 FLocomotor::AddFootToSet(
 	}
 	
 	const int32 FootIndexInSet = FootSets[FootSetIndex].AddFoot(InitialRootGoalWorld, InitialWorldFoot, InSettings);
-	AllFeet.Add(&FootSets[FootSetIndex].GetFeet()[FootIndexInSet]);
+
+	// update cached pointers to all feet
+	// NOTE: we reset and update all pointers here in case the foot arrays inside the sets
+	// were reallocated which would invalidate all the pointers in the cache.
+	AllFeet.Reset(0);
+	for (FLocomotorFootSet& FootSet : FootSets)
+	{
+		for (FLocomotorFoot& Foot : FootSet.GetFeet())
+		{
+			AllFeet.Add(&Foot);
+		}
+	}
+	
 	return FootIndexInSet;
 }
 
@@ -365,7 +378,8 @@ void FLocomotor::UpdateFeetTargets()
 		}
 		else
 		{
-			const float PhaseOffsetOfSet = FootSets[Foot.Settings.FootSetIndex].GetPhaseOffset();
+			const int32 FootSetIndex = FMath::Clamp(Foot.Settings.FootSetIndex, 0, FootSets.Num() - 1);
+			const float PhaseOffsetOfSet = FootSets[FootSetIndex].GetPhaseOffset();
 			const float StaticOffsetOfFoot = Foot.Settings.StaticPhaseOffset;
 			Foot.TargetPhase = WrapPhaseInRange(CurrentGlobalPhase + StaticOffsetOfFoot + PhaseOffsetOfSet);
 			Foot.CurrentPhase = BlendTowardsTargetPhase(Foot.CurrentPhase, Foot.TargetPhase, PhaseBlendRate);
